Moved Niagara to its own module; made curve access work; fixed various small things

diff --git a/Engine/Source/Developer/AssetTools/AssetTools.Build.cs b/Engine/Source/Developer/AssetTools/AssetTools.Build.cs
--- a/Engine/Source/Developer/AssetTools/AssetTools.Build.cs
+++ b/Engine/Source/Developer/AssetTools/AssetTools.Build.cs
@@ -24,6 +24,7 @@ public class AssetTools : ModuleRules
 				"Kismet",
 				"Landscape",
                 "Foliage",
+                "Niagara",
 			}
 		);
 
diff --git a/Engine/Source/Developer/AssetTools/Private/AssetTypeActions/AssetTypeActions_NiagaraEffect.cpp b/Engine/Source/Developer/AssetTools/Private/AssetTypeActions/AssetTypeActions_NiagaraEffect.cpp
--- a/Engine/Source/Developer/AssetTools/Private/AssetTypeActions/AssetTypeActions_NiagaraEffect.cpp
+++ b/Engine/Source/Developer/AssetTools/Private/AssetTypeActions/AssetTypeActions_NiagaraEffect.cpp
@@ -1,7 +1,7 @@
 // Copyright 1998-2015 Epic Games, Inc. All Rights Reserved.
 
 #include "AssetToolsPrivatePCH.h"
-#include "Engine/NiagaraEffect.h"
+#include "NiagaraEffect.h"
 
 #include "Toolkits/AssetEditorManager.h"
 
diff --git a/Engine/Source/Developer/AssetTools/Private/AssetTypeActions/AssetTypeActions_NiagaraScript.cpp b/Engine/Source/Developer/AssetTools/Private/AssetTypeActions/AssetTypeActions_NiagaraScript.cpp
--- a/Engine/Source/Developer/AssetTools/Private/AssetTypeActions/AssetTypeActions_NiagaraScript.cpp
+++ b/Engine/Source/Developer/AssetTools/Private/AssetTypeActions/AssetTypeActions_NiagaraScript.cpp
@@ -1,7 +1,7 @@
 // Copyright 1998-2015 Epic Games, Inc. All Rights Reserved.
 
 #include "AssetToolsPrivatePCH.h"
-#include "Engine/NiagaraScript.h"
+#include "NiagaraScript.h"
 
 #include "Toolkits/AssetEditorManager.h"
 
diff --git a/Engine/Source/Editor/NiagaraEditor/NiagaraEditor.Build.cs b/Engine/Source/Editor/NiagaraEditor/NiagaraEditor.Build.cs
--- a/Engine/Source/Editor/NiagaraEditor/NiagaraEditor.Build.cs
+++ b/Engine/Source/Editor/NiagaraEditor/NiagaraEditor.Build.cs
@@ -19,13 +19,15 @@ public class NiagaraEditor : ModuleRules
                 "EditorStyle",
 				"UnrealEd", 
 				"GraphEditor", 
-				"VectorVM"
+				"VectorVM",
+                "Niagara"
 			}
 		);
 
         PublicDependencyModuleNames.AddRange(
             new string[] {
-				"Engine"
+				"Engine",
+                "Niagara"
             }
         );
 
diff --git a/Engine/Source/Editor/NiagaraEditor/Private/NiagaraCompiler.cpp b/Engine/Source/Editor/NiagaraEditor/Private/NiagaraCompiler.cpp
--- a/Engine/Source/Editor/NiagaraEditor/Private/NiagaraCompiler.cpp
+++ b/Engine/Source/Editor/NiagaraEditor/Private/NiagaraCompiler.cpp
@@ -1,8 +1,8 @@
 // Copyright 1998-2015 Epic Games, Inc. All Rights Reserved.
 
 #include "NiagaraEditorPrivatePCH.h"
-#include "Engine/NiagaraScript.h"
-#include "Components/NiagaraComponent.h"
+#include "NiagaraScript.h"
+#include "NiagaraComponent.h"
 #include "CompilerResultsLog.h"
 #include "EdGraphUtilities.h"
 #include "VectorVM.h"
diff --git a/Engine/Source/Editor/NiagaraEditor/Private/NiagaraCompiler.h b/Engine/Source/Editor/NiagaraEditor/Private/NiagaraCompiler.h
--- a/Engine/Source/Editor/NiagaraEditor/Private/NiagaraCompiler.h
+++ b/Engine/Source/Editor/NiagaraEditor/Private/NiagaraCompiler.h
@@ -4,7 +4,7 @@
 
 #include "INiagaraCompiler.h"
 #include "CompilerResultsLog.h"
-#include "Niagara/NiagaraScriptConstantData.h"
+#include "NiagaraScriptConstantData.h"
 
 /** Base class for Niagara compilers. Children of this will include a compiler for the VectorVM and for Compute shaders. Possibly others. */
 class NIAGARAEDITOR_API FNiagaraCompiler : public INiagaraCompiler
diff --git a/Engine/Source/Editor/NiagaraEditor/Private/NiagaraCompiler_VectorVM.cpp b/Engine/Source/Editor/NiagaraEditor/Private/NiagaraCompiler_VectorVM.cpp
--- a/Engine/Source/Editor/NiagaraEditor/Private/NiagaraCompiler_VectorVM.cpp
+++ b/Engine/Source/Editor/NiagaraEditor/Private/NiagaraCompiler_VectorVM.cpp
@@ -1,8 +1,8 @@
 // Copyright 1998-2015 Epic Games, Inc. All Rights Reserved.
 
 #include "NiagaraEditorPrivatePCH.h"
-#include "Engine/NiagaraScript.h"
-#include "Components/NiagaraComponent.h"
+#include "NiagaraScript.h"
+#include "NiagaraComponent.h"
 #include "CompilerResultsLog.h"
 #include "EdGraphUtilities.h"
 #include "VectorVM.h"
@@ -621,4 +621,16 @@ void FNiagaraCompiler_VectorVM::EventBroadcast_Internal(TArray<TNiagaraExprPtr>&
 	OutputExpressions.Add(Expression_VMNative(VectorVM::EOp::eventbroadcast, InputExpressions));
 }
 
+void FNiagaraCompiler_VectorVM::EaseIn_Internal(TArray<TNiagaraExprPtr>& InputExpressions, TArray<TNiagaraExprPtr>& OutputExpressions)
+{
+	OutputExpressions.Add(Expression_VMNative(VectorVM::EOp::easein, InputExpressions));
+}
+
+
+void FNiagaraCompiler_VectorVM::EaseInOut_Internal(TArray<TNiagaraExprPtr>& InputExpressions, TArray<TNiagaraExprPtr>& OutputExpressions)
+{
+	OutputExpressions.Add(Expression_VMNative(VectorVM::EOp::easeinout, InputExpressions));
+}
+
+
 #undef LOCTEXT_NAMESPACE
diff --git a/Engine/Source/Editor/NiagaraEditor/Private/NiagaraEditor.cpp b/Engine/Source/Editor/NiagaraEditor/Private/NiagaraEditor.cpp
--- a/Engine/Source/Editor/NiagaraEditor/Private/NiagaraEditor.cpp
+++ b/Engine/Source/Editor/NiagaraEditor/Private/NiagaraEditor.cpp
@@ -1,7 +1,7 @@
 // Copyright 1998-2015 Epic Games, Inc. All Rights Reserved.
 
 #include "NiagaraEditorPrivatePCH.h"
-#include "Engine/NiagaraScript.h"
+#include "NiagaraScript.h"
 
 #include "Toolkits/IToolkitHost.h"
 #include "Editor/WorkspaceMenuStructure/Public/WorkspaceMenuStructureModule.h"
diff --git a/Engine/Source/Editor/NiagaraEditor/Private/NiagaraEffectEditor.cpp b/Engine/Source/Editor/NiagaraEditor/Private/NiagaraEffectEditor.cpp
--- a/Engine/Source/Editor/NiagaraEditor/Private/NiagaraEffectEditor.cpp
+++ b/Engine/Source/Editor/NiagaraEditor/Private/NiagaraEffectEditor.cpp
@@ -1,7 +1,7 @@
 // Copyright 1998-2015 Epic Games, Inc. All Rights Reserved.
 
 #include "NiagaraEditorPrivatePCH.h"
-#include "Engine/NiagaraEffect.h"
+#include "NiagaraEffect.h"
 
 #include "Toolkits/IToolkitHost.h"
 #include "Editor/WorkspaceMenuStructure/Public/WorkspaceMenuStructureModule.h"
diff --git a/Engine/Source/Editor/NiagaraEditor/Private/SNiagaraEffectEditorViewport.cpp b/Engine/Source/Editor/NiagaraEditor/Private/SNiagaraEffectEditorViewport.cpp
--- a/Engine/Source/Editor/NiagaraEditor/Private/SNiagaraEffectEditorViewport.cpp
+++ b/Engine/Source/Editor/NiagaraEditor/Private/SNiagaraEffectEditorViewport.cpp
@@ -8,8 +8,8 @@
 #include "Runtime/Engine/Public/Slate/SceneViewport.h"
 #include "NiagaraEditor.h"
 #include "ComponentReregisterContext.h"
-#include "Components/NiagaraComponent.h"
-#include "Engine/NiagaraEffect.h"
+#include "NiagaraComponent.h"
+#include "NiagaraEffect.h"
 #include "SDockTab.h"
 #include "Engine/TextureCube.h"
 
diff --git a/Engine/Source/Editor/NiagaraEditor/Private/SNiagaraEffectEditorWidget.h b/Engine/Source/Editor/NiagaraEditor/Private/SNiagaraEffectEditorWidget.h
--- a/Engine/Source/Editor/NiagaraEditor/Private/SNiagaraEffectEditorWidget.h
+++ b/Engine/Source/Editor/NiagaraEditor/Private/SNiagaraEffectEditorWidget.h
@@ -11,9 +11,9 @@
 #include "SCurveEditor.h"
 #include "AssetThumbnail.h"
 #include "SNiagaraEffectEditorViewport.h"
-#include "Components/NiagaraComponent.h"
-#include "Engine/NiagaraEffect.h"
-#include "Engine/NiagaraSimulation.h"
+#include "NiagaraComponent.h"
+#include "NiagaraEffect.h"
+#include "NiagaraSimulation.h"
 #include "Engine/NiagaraEffectRenderer.h"
 #include "ComponentReregisterContext.h"
 #include "SNumericEntryBox.h"
diff --git a/Engine/Source/Editor/UnrealEd/Classes/Niagara/NiagaraNodeConstant.h b/Engine/Source/Editor/UnrealEd/Classes/Niagara/NiagaraNodeConstant.h
--- a/Engine/Source/Editor/UnrealEd/Classes/Niagara/NiagaraNodeConstant.h
+++ b/Engine/Source/Editor/UnrealEd/Classes/Niagara/NiagaraNodeConstant.h
@@ -2,7 +2,7 @@
 
 #pragma once
 
-#include "Niagara/NiagaraCommon.h"
+#include "NiagaraCommon.h"
 #include "NiagaraNodeConstant.generated.h"
 
 UCLASS(MinimalAPI, Deprecated)
diff --git a/Engine/Source/Editor/UnrealEd/Classes/Niagara/NiagaraScriptSource.h b/Engine/Source/Editor/UnrealEd/Classes/Niagara/NiagaraScriptSource.h
--- a/Engine/Source/Editor/UnrealEd/Classes/Niagara/NiagaraScriptSource.h
+++ b/Engine/Source/Editor/UnrealEd/Classes/Niagara/NiagaraScriptSource.h
@@ -2,7 +2,7 @@
 
 
 #pragma once
-#include "Engine/NiagaraScriptSourceBase.h"
+#include "NiagaraScriptSourceBase.h"
 #include "NiagaraScriptSource.generated.h"
 
 
diff --git a/Engine/Source/Editor/UnrealEd/Private/Factories/ActorFactory.cpp b/Engine/Source/Editor/UnrealEd/Private/Factories/ActorFactory.cpp
--- a/Engine/Source/Editor/UnrealEd/Private/Factories/ActorFactory.cpp
+++ b/Engine/Source/Editor/UnrealEd/Private/Factories/ActorFactory.cpp
@@ -26,8 +26,8 @@ ActorFactory.cpp:
 #include "Particles/ParticleSystem.h"
 #include "Particles/ParticleSystemComponent.h"
 
-#include "Engine/NiagaraActor.h"
-#include "Engine/NiagaraEffect.h"
+#include "NiagaraActor.h"
+#include "NiagaraEffect.h"
 
 #include "VectorField/VectorField.h"
 
diff --git a/Engine/Source/Editor/UnrealEd/Private/Factories/NiagaraEffectFactory.cpp b/Engine/Source/Editor/UnrealEd/Private/Factories/NiagaraEffectFactory.cpp
--- a/Engine/Source/Editor/UnrealEd/Private/Factories/NiagaraEffectFactory.cpp
+++ b/Engine/Source/Editor/UnrealEd/Private/Factories/NiagaraEffectFactory.cpp
@@ -1,7 +1,7 @@
 // Copyright 1998-2015 Epic Games, Inc. All Rights Reserved.
 
 #include "UnrealEd.h"
-#include "Engine/NiagaraEffect.h"
+#include "NiagaraEffect.h"
 #include "NiagaraEditorModule.h"
 
 #define LOCTEXT_NAMESPACE "NiagaraEffectFactory"
diff --git a/Engine/Source/Editor/UnrealEd/Private/Factories/NiagaraScriptFactory.cpp b/Engine/Source/Editor/UnrealEd/Private/Factories/NiagaraScriptFactory.cpp
--- a/Engine/Source/Editor/UnrealEd/Private/Factories/NiagaraScriptFactory.cpp
+++ b/Engine/Source/Editor/UnrealEd/Private/Factories/NiagaraScriptFactory.cpp
@@ -1,7 +1,7 @@
 // Copyright 1998-2015 Epic Games, Inc. All Rights Reserved.
 
 #include "UnrealEd.h"
-#include "Engine/NiagaraScript.h"
+#include "NiagaraScript.h"
 #include "NiagaraEditorModule.h"
 
 #define LOCTEXT_NAMESPACE "NiagaraScriptFactory"
diff --git a/Engine/Source/Editor/UnrealEd/Private/Niagara/EdGraphSchema_Niagara.cpp b/Engine/Source/Editor/UnrealEd/Private/Niagara/EdGraphSchema_Niagara.cpp
--- a/Engine/Source/Editor/UnrealEd/Private/Niagara/EdGraphSchema_Niagara.cpp
+++ b/Engine/Source/Editor/UnrealEd/Private/Niagara/EdGraphSchema_Niagara.cpp
@@ -4,7 +4,7 @@
 
 #include "BlueprintGraphDefinitions.h"
 #include "GraphEditorSettings.h"
-#include "Components/NiagaraComponent.h"
+#include "NiagaraComponent.h"
 
 #include "INiagaraEditor.h"
 #include "INiagaraCompiler.h"
diff --git a/Engine/Source/Editor/UnrealEd/Private/Niagara/NiagaraEditorCommon.cpp b/Engine/Source/Editor/UnrealEd/Private/Niagara/NiagaraEditorCommon.cpp
--- a/Engine/Source/Editor/UnrealEd/Private/Niagara/NiagaraEditorCommon.cpp
+++ b/Engine/Source/Editor/UnrealEd/Private/Niagara/NiagaraEditorCommon.cpp
@@ -540,10 +540,28 @@ void FNiagaraOpInfo::Init()
 	Op->Name = EventBroadcast;
 	Op->FriendlyName = NSLOCTEXT("NiagaraOpInfo", "EventBroadcast Name", "Broadcast Event");
 	Op->Description = NSLOCTEXT("NiagaraOpInfo", "EventBroadcast Desc", "Broadcast Event if trigger is >=1.0");
-	Op->Inputs.Add(FNiagaraOpInOutInfo(A, ENiagaraDataType::Curve, NSLOCTEXT("NiagaraOpInfo", "Trigger", "Trigger"), NSLOCTEXT("NiagaraOpInfo", "Trigger", "Trigger"), DefaultStr_VecOne));
+	Op->Inputs.Add(FNiagaraOpInOutInfo(A, ENiagaraDataType::Vector, NSLOCTEXT("NiagaraOpInfo", "Trigger", "Trigger"), NSLOCTEXT("NiagaraOpInfo", "Trigger", "Trigger"), DefaultStr_VecOne));
 //	Op->Outputs.Add(FNiagaraOpInOutInfo(Result, ENiagaraDataType::Vector, ResultText, ResultText, DefaultStr_VecOne));
 	Op->OpDelegate.BindStatic(&INiagaraCompiler::EventBroadcast);
 
+	Op = &OpInfos.Add(EaseIn);
+	Op->Name = EaseIn;
+	Op->FriendlyName = NSLOCTEXT("NiagaraOpInfo", "EaseIn Name", "Ease In 0->1");
+	Op->Description = NSLOCTEXT("NiagaraOpInfo", "EaseIn Desc", "Ease in for Time; A and B determine the ease length in [0:1]");
+	Op->Inputs.Add(FNiagaraOpInOutInfo(A, ENiagaraDataType::Vector, NSLOCTEXT("NiagaraOpInfo", "A", "A"), NSLOCTEXT("NiagaraOpInfo", "A", "A"), DefaultStr_VecZero));
+	Op->Inputs.Add(FNiagaraOpInOutInfo(A, ENiagaraDataType::Vector, NSLOCTEXT("NiagaraOpInfo", "B", "B"), NSLOCTEXT("NiagaraOpInfo", "B", "B"), DefaultStr_VecOne));
+	Op->Inputs.Add(FNiagaraOpInOutInfo(A, ENiagaraDataType::Vector, NSLOCTEXT("NiagaraOpInfo", "Time", "Time"), NSLOCTEXT("NiagaraOpInfo", "Time", "Time"), DefaultStr_VecOne));
+	Op->Outputs.Add(FNiagaraOpInOutInfo(Result, ENiagaraDataType::Vector, ResultText, ResultText, DefaultStr_VecOne));
+	Op->OpDelegate.BindStatic(&INiagaraCompiler::EaseIn);
+
+	Op = &OpInfos.Add(EaseInOut);
+	Op->Name = EaseInOut;
+	Op->FriendlyName = NSLOCTEXT("NiagaraOpInfo", "EaseInOut Name", "Ease In/Out 0->1->0");
+	Op->Description = NSLOCTEXT("NiagaraOpInfo", "EaseInOut Desc", "Ease in/out for Time; A and B determine the ease in/out lengths in [0:1]");
+	Op->Inputs.Add(FNiagaraOpInOutInfo(A, ENiagaraDataType::Vector, NSLOCTEXT("NiagaraOpInfo", "Time", "Time"), NSLOCTEXT("NiagaraOpInfo", "Time", "Time"), DefaultStr_VecOne));
+	Op->Outputs.Add(FNiagaraOpInOutInfo(Result, ENiagaraDataType::Vector, ResultText, ResultText, DefaultStr_VecOne));
+	Op->OpDelegate.BindStatic(&INiagaraCompiler::EaseInOut);
+
 
 	#define NiagaraOp(OpName) if (GetOpInfo(OpName) == NULL)\
 	{\
diff --git a/Engine/Source/Editor/UnrealEd/Private/Niagara/NiagaraNodeFunctionCall.cpp b/Engine/Source/Editor/UnrealEd/Private/Niagara/NiagaraNodeFunctionCall.cpp
--- a/Engine/Source/Editor/UnrealEd/Private/Niagara/NiagaraNodeFunctionCall.cpp
+++ b/Engine/Source/Editor/UnrealEd/Private/Niagara/NiagaraNodeFunctionCall.cpp
@@ -3,7 +3,7 @@
 #include "UnrealEd.h"
 #include "BlueprintGraphDefinitions.h"
 #include "GraphEditorSettings.h"
-#include "Engine/NiagaraScript.h"
+#include "NiagaraScript.h"
 
 void UNiagaraNodeFunctionCall::PostLoad()
 {
diff --git a/Engine/Source/Editor/UnrealEd/Private/Niagara/NiagaraScriptSource.cpp b/Engine/Source/Editor/UnrealEd/Private/Niagara/NiagaraScriptSource.cpp
--- a/Engine/Source/Editor/UnrealEd/Private/Niagara/NiagaraScriptSource.cpp
+++ b/Engine/Source/Editor/UnrealEd/Private/Niagara/NiagaraScriptSource.cpp
@@ -1,8 +1,8 @@
 // Copyright 1998-2015 Epic Games, Inc. All Rights Reserved.
 
 #include "UnrealEd.h"
-#include "Engine/NiagaraScript.h"
-#include "Engine/NiagaraConstants.h"
+#include "NiagaraScript.h"
+#include "NiagaraConstants.h"
 #include "NiagaraEditorModule.h"
 
 //////////////////////////////////////////////////////////////////////////
@@ -89,8 +89,8 @@ void UNiagaraScriptSource::PostLoad()
 		ScriptOwner->ConditionalPostLoad();
 		//FNiagaraEditorModule& NiagaraEditorModule = FModuleManager::Get().LoadModuleChecked<FNiagaraEditorModule>(TEXT("NiagaraEditor"));
 		//NiagaraEditorModule.CompileScript(ScriptOwner);	
-		Compile();
 	}
+	Compile();
 #endif
 }
 
diff --git a/Engine/Source/Editor/UnrealEd/Public/NiagaraEditorCommon.h b/Engine/Source/Editor/UnrealEd/Public/NiagaraEditorCommon.h
--- a/Engine/Source/Editor/UnrealEd/Public/NiagaraEditorCommon.h
+++ b/Engine/Source/Editor/UnrealEd/Public/NiagaraEditorCommon.h
@@ -2,7 +2,8 @@
 
 #pragma once
 
-#include "Niagara/NiagaraCommon.h"
+#include "NiagaraCommon.h"
+
 
 //List of operations supported by the Niagara compiler(s) as visible from the outside.
 //Additional information about each operation is added in FNiagaraOpInfo::Init
@@ -65,6 +66,8 @@
 	NiagaraOp(Sample)\
 	NiagaraOp(Write)\
 	NiagaraOp(EventBroadcast)\
+	NiagaraOp(EaseIn)\
+	NiagaraOp(EaseInOut)\
 
 enum class ENiagaraExpressionResultLocation
 {
@@ -76,6 +79,8 @@ enum class ENiagaraExpressionResultLocation
 	Unknown,
 };
 
+
+
 typedef TSharedPtr<class FNiagaraExpression> TNiagaraExprPtr;
 
 /** Base class for all Niagara Expresions. Intermediate data collected from the node graph used for compiling Niagara scripts. */
@@ -200,4 +205,5 @@ public:
 
 	TNiagaraExprPtr Expression;
 	UEdGraphPin* OutputPin;
-};
\ No newline at end of file
+};
+
diff --git a/Engine/Source/Editor/UnrealEd/UnrealEd.Build.cs b/Engine/Source/Editor/UnrealEd/UnrealEd.Build.cs
--- a/Engine/Source/Editor/UnrealEd/UnrealEd.Build.cs
+++ b/Engine/Source/Editor/UnrealEd/UnrealEd.Build.cs
@@ -53,6 +53,7 @@ public class UnrealEd : ModuleRules
                 "SuperSearch",
 				"OutputLog",
 				"Landscape",
+                "Niagara",
 			}
 		);
 
@@ -77,6 +78,7 @@ public class UnrealEd : ModuleRules
 				"BlueprintGraph",
                 "Http",
 				"UnrealAudio",
+                "Niagara",
 			}
 		);
 
@@ -220,6 +222,7 @@ public class UnrealEd : ModuleRules
              	"CollectionManager",
 				"BlueprintGraph",
 				"NiagaraEditor",
+                "Niagara",
 			}
 			);
 
diff --git a/Engine/Source/Runtime/Engine/Classes/Components/NiagaraComponent.h b/Engine/Source/Runtime/Engine/Classes/Components/NiagaraComponent.h
deleted file mode 100644
diff --git a/Engine/Source/Runtime/Engine/Classes/Engine/NiagaraActor.h b/Engine/Source/Runtime/Engine/Classes/Engine/NiagaraActor.h
deleted file mode 100644
diff --git a/Engine/Source/Runtime/Engine/Classes/Engine/NiagaraConstants.h b/Engine/Source/Runtime/Engine/Classes/Engine/NiagaraConstants.h
deleted file mode 100644
diff --git a/Engine/Source/Runtime/Engine/Classes/Engine/NiagaraEffect.h b/Engine/Source/Runtime/Engine/Classes/Engine/NiagaraEffect.h
deleted file mode 100644
diff --git a/Engine/Source/Runtime/Engine/Classes/Engine/NiagaraEffectRenderer.h b/Engine/Source/Runtime/Engine/Classes/Engine/NiagaraEffectRenderer.h
--- a/Engine/Source/Runtime/Engine/Classes/Engine/NiagaraEffectRenderer.h
+++ b/Engine/Source/Runtime/Engine/Classes/Engine/NiagaraEffectRenderer.h
@@ -9,7 +9,7 @@ NiagaraEffectRenderer.h: Base class for Niagara render modules
 #include "ParticleHelper.h"
 #include "ParticleVertexFactory.h"
 #include "ParticleBeamTrailVertexFactory.h"
-#include "Components/NiagaraComponent.h"
+#include "NiagaraComponent.h"
 #include "NiagaraEffectRendererProperties.h"
 #include "NiagaraSpriteRendererProperties.h"
 #include "NiagaraRibbonRendererProperties.h"
@@ -123,7 +123,7 @@ private:
 /**
 * NiagaraEffectRendererSprites renders an FNiagaraSimulation as sprite particles
 */
-class NiagaraEffectRendererSprites : public NiagaraEffectRenderer
+class ENGINE_API NiagaraEffectRendererSprites : public NiagaraEffectRenderer
 {
 public:	
 
@@ -167,7 +167,7 @@ private:
 * NiagaraEffectRendererRibbon renders an FNiagaraSimulation as a ribbon connecting all particles
 * in order by particle age.
 */
-class NiagaraEffectRendererRibbon : public NiagaraEffectRenderer
+class ENGINE_API NiagaraEffectRendererRibbon : public NiagaraEffectRenderer
 {
 public:
 	NiagaraEffectRendererRibbon(ERHIFeatureLevel::Type FeatureLevel, UNiagaraEffectRendererProperties *Props);
diff --git a/Engine/Source/Runtime/Engine/Classes/Engine/NiagaraEffectRendererProperties.h b/Engine/Source/Runtime/Engine/Classes/Engine/NiagaraEffectRendererProperties.h
--- a/Engine/Source/Runtime/Engine/Classes/Engine/NiagaraEffectRendererProperties.h
+++ b/Engine/Source/Runtime/Engine/Classes/Engine/NiagaraEffectRendererProperties.h
@@ -10,7 +10,7 @@
 * and handed back to the effect renderer on load.
 */
 UCLASS()
-class UNiagaraEffectRendererProperties : public UObject
+class ENGINE_API UNiagaraEffectRendererProperties : public UObject
 {
 	GENERATED_UCLASS_BODY()
 	UNiagaraEffectRendererProperties()
diff --git a/Engine/Source/Runtime/Engine/Classes/Engine/NiagaraScript.h b/Engine/Source/Runtime/Engine/Classes/Engine/NiagaraScript.h
deleted file mode 100644
diff --git a/Engine/Source/Runtime/Engine/Classes/Engine/NiagaraScriptSourceBase.h b/Engine/Source/Runtime/Engine/Classes/Engine/NiagaraScriptSourceBase.h
deleted file mode 100644
diff --git a/Engine/Source/Runtime/Engine/Classes/Engine/NiagaraSimulation.h b/Engine/Source/Runtime/Engine/Classes/Engine/NiagaraSimulation.h
deleted file mode 100644
diff --git a/Engine/Source/Runtime/Engine/Classes/Niagara/NiagaraCommon.h b/Engine/Source/Runtime/Engine/Classes/Niagara/NiagaraCommon.h
deleted file mode 100644
diff --git a/Engine/Source/Runtime/Engine/Classes/Niagara/NiagaraConstantSet.h b/Engine/Source/Runtime/Engine/Classes/Niagara/NiagaraConstantSet.h
deleted file mode 100644
diff --git a/Engine/Source/Runtime/Engine/Classes/Niagara/NiagaraScriptConstantData.h b/Engine/Source/Runtime/Engine/Classes/Niagara/NiagaraScriptConstantData.h
deleted file mode 100644
diff --git a/Engine/Source/Runtime/Engine/Engine.Build.cs b/Engine/Source/Runtime/Engine/Engine.Build.cs
--- a/Engine/Source/Runtime/Engine/Engine.Build.cs
+++ b/Engine/Source/Runtime/Engine/Engine.Build.cs
@@ -55,7 +55,6 @@ public class Engine : ModuleRules
 				"EngineSettings",
 				"SynthBenchmark",
                 "AIModule",
-				"VectorVM",
 				"DatabaseSupport",
 			}
 		);
@@ -66,7 +65,6 @@ public class Engine : ModuleRules
 				"Networking",
 				"Sockets",
 				"SlateReflector",
-				"VectorVM",
 				"Landscape",
                 "UMG",
 				"Projects"
diff --git a/Engine/Source/Runtime/Engine/Private/Components/NiagaraComponent.cpp b/Engine/Source/Runtime/Engine/Private/Components/NiagaraComponent.cpp
deleted file mode 100644
diff --git a/Engine/Source/Runtime/Engine/Private/InstancedStaticMesh.h b/Engine/Source/Runtime/Engine/Private/InstancedStaticMesh.h
--- a/Engine/Source/Runtime/Engine/Private/InstancedStaticMesh.h
+++ b/Engine/Source/Runtime/Engine/Private/InstancedStaticMesh.h
@@ -13,7 +13,6 @@
 #include "Misc/UObjectToken.h"
 #include "Components/SplineMeshComponent.h"
 #include "Components/ModelComponent.h"
-#include "Components/NiagaraComponent.h"
 #include "Components/ShapeComponent.h"
 #include "Components/BoxComponent.h"
 #include "Components/CapsuleComponent.h"
diff --git a/Engine/Source/Runtime/Engine/Private/NiagaraActor.cpp b/Engine/Source/Runtime/Engine/Private/NiagaraActor.cpp
deleted file mode 100644
diff --git a/Engine/Source/Runtime/Engine/Private/NiagaraConstantSet.cpp b/Engine/Source/Runtime/Engine/Private/NiagaraConstantSet.cpp
deleted file mode 100644
diff --git a/Engine/Source/Runtime/Engine/Private/NiagaraEffect.cpp b/Engine/Source/Runtime/Engine/Private/NiagaraEffect.cpp
deleted file mode 100644
diff --git a/Engine/Source/Runtime/Engine/Private/NiagaraScript.cpp b/Engine/Source/Runtime/Engine/Private/NiagaraScript.cpp
deleted file mode 100644
diff --git a/Engine/Source/Runtime/Engine/Private/NiagaraSimulation.cpp b/Engine/Source/Runtime/Engine/Private/NiagaraSimulation.cpp
deleted file mode 100644
diff --git a/Engine/Source/Runtime/Engine/Public/ParticleBeamTrailVertexFactory.h b/Engine/Source/Runtime/Engine/Public/ParticleBeamTrailVertexFactory.h
--- a/Engine/Source/Runtime/Engine/Public/ParticleBeamTrailVertexFactory.h
+++ b/Engine/Source/Runtime/Engine/Public/ParticleBeamTrailVertexFactory.h
@@ -22,7 +22,7 @@ typedef TUniformBufferRef<FParticleBeamTrailUniformParameters> FParticleBeamTrai
 /**
  * Beam/Trail particle vertex factory.
  */
-class FParticleBeamTrailVertexFactory : public FParticleVertexFactoryBase
+class ENGINE_API FParticleBeamTrailVertexFactory : public FParticleVertexFactoryBase
 {
 	DECLARE_VERTEX_FACTORY_TYPE(FParticleBeamTrailVertexFactory);
 
diff --git a/Engine/Source/Runtime/Engine/Public/ParticleVertexFactory.h b/Engine/Source/Runtime/Engine/Public/ParticleVertexFactory.h
--- a/Engine/Source/Runtime/Engine/Public/ParticleVertexFactory.h
+++ b/Engine/Source/Runtime/Engine/Public/ParticleVertexFactory.h
@@ -97,7 +97,7 @@ typedef TUniformBufferRef<FParticleSpriteUniformParameters> FParticleSpriteUnifo
 /**
  * Vertex factory for rendering particle sprites.
  */
-class FParticleSpriteVertexFactory : public FParticleVertexFactoryBase
+class ENGINE_API FParticleSpriteVertexFactory : public FParticleVertexFactoryBase
 {
 	DECLARE_VERTEX_FACTORY_TYPE(FParticleSpriteVertexFactory);
 
diff --git a/Engine/Source/Runtime/Niagara/Classes/NiagaraConstants.h b/Engine/Source/Runtime/Niagara/Classes/NiagaraConstants.h
new file mode 100644
--- /dev/null
+++ b/Engine/Source/Runtime/Niagara/Classes/NiagaraConstants.h
@@ -0,0 +1,30 @@
+// Copyright 1998-2015 Epic Games, Inc. All Rights Reserved.
+
+#pragma once
+
+/**  List of built-in constants
+*	The values of these enums correspond directly to indices for the compiler and interpreter
+*/
+namespace NiagaraConstants
+{
+	enum {
+		DeltaSeconds = 1,
+		EmitterPosition,
+		EmitterAge,
+		ActorXAxis,
+		ActorYAxis,
+		ActorZAxis,
+		NumBuiltinConstants
+	} ConstantDef;
+
+	FName ConstantNames[] = {
+		"Delta Time",
+		"Emitter position",
+		"Emitter age",
+		"Emitter X Axis",
+		"Emitter Y Axis",
+		"Emitter Z Axis",
+		""
+	};
+};
+
diff --git a/Engine/Source/Runtime/Niagara/Classes/NiagaraEffect.h b/Engine/Source/Runtime/Niagara/Classes/NiagaraEffect.h
new file mode 100644
--- /dev/null
+++ b/Engine/Source/Runtime/Niagara/Classes/NiagaraEffect.h
@@ -0,0 +1,161 @@
+// Copyright 1998-2015 Epic Games, Inc. All Rights Reserved.
+
+#pragma once
+#include "NiagaraSimulation.h"
+#include "Engine/NiagaraEffectRenderer.h"
+#include "List.h"
+#include "Runtime/VectorVM/Public/VectorVMDataObject.h"
+#include "NiagaraEffect.generated.h"
+
+
+UCLASS(MinimalAPI)
+class UNiagaraEffect : public UObject
+{
+	GENERATED_UCLASS_BODY()
+
+public:
+	NIAGARA_API FNiagaraEmitterProperties* AddEmitterProperties();
+
+	void Init()
+	{
+	}
+
+
+	FNiagaraEmitterProperties *GetEmitterProperties(int Idx)
+	{
+		check(Idx < EmitterProps.Num());
+		return EmitterProps[Idx];
+	};
+
+	int GetNumEmitters()
+	{
+		return EmitterProps.Num();
+	}
+
+	NIAGARA_API void CreateEffectRendererProps(TSharedPtr<FNiagaraSimulation> Sim);
+
+	virtual void PostLoad() override; 
+	virtual void PreSave() override;
+
+private:
+	// serialized array of emitter properties
+	UPROPERTY()
+	TArray<FNiagaraEmitterProperties>EmitterPropsSerialized;
+
+	TArray<FNiagaraEmitterProperties*> EmitterProps;
+};
+
+
+
+class FNiagaraEffectInstance
+{
+public:
+	explicit FNiagaraEffectInstance(UNiagaraEffect *InAsset, UNiagaraComponent *InComponent)
+	{
+		Component = InComponent;
+		for (int32 i = 0; i < InAsset->GetNumEmitters(); i++)
+		{
+			FNiagaraSimulation *Sim = new FNiagaraSimulation(InAsset->GetEmitterProperties(i), Component->GetWorld()->FeatureLevel);
+			Emitters.Add(MakeShareable(Sim));
+		}
+		InitRenderModules(Component->GetWorld()->FeatureLevel);
+	}
+
+	explicit FNiagaraEffectInstance(UNiagaraEffect *InAsset)
+		: Component(nullptr)
+	{
+		InitEmitters(InAsset);
+	}
+
+	NIAGARA_API void InitEmitters(UNiagaraEffect *InAsset);
+	void InitRenderModules(ERHIFeatureLevel::Type InFeatureLevel)
+	{
+		for (TSharedPtr<FNiagaraSimulation> Sim : Emitters)
+		{
+			if(Sim->GetProperties()->RendererProperties == nullptr)
+			{
+				Component->GetAsset()->CreateEffectRendererProps(Sim);
+			}
+			Sim->SetRenderModuleType(Sim->GetProperties()->RenderModuleType, InFeatureLevel);
+		}
+	}
+
+	void Init(UNiagaraComponent *InComponent)
+	{
+		if (InComponent->GetAsset())
+		{
+			Emitters.Empty();
+			InitEmitters(InComponent->GetAsset());
+		}
+		Component = InComponent;
+		InitRenderModules(Component->GetWorld()->FeatureLevel);
+		RenderModuleupdate();
+	}
+
+	NIAGARA_API TSharedPtr<FNiagaraSimulation> AddEmitter(FNiagaraEmitterProperties *Properties);
+
+	void SetConstant(FNiagaraVariableInfo ID, const float Value)
+	{
+		Constants.SetOrAdd(ID, Value);
+	}
+
+	void SetConstant(FNiagaraVariableInfo ID, const FVector4& Value)
+	{
+		Constants.SetOrAdd(ID, Value);
+	}
+
+	void SetConstant(FNiagaraVariableInfo ID, const FMatrix& Value)
+	{
+		Constants.SetOrAdd(ID, Value);
+	}
+
+	void SetConstant(FName ConstantName, FNiagaraDataObject *Value)
+	{
+		Constants.SetOrAdd(ConstantName, Value);
+	}
+
+
+
+	void Tick(float DeltaSeconds)
+	{
+		Constants.SetOrAdd(FName(TEXT("EffectGrid")), &VolumeGrid); 
+
+		// pass the constants down to the emitter
+		// TODO: should probably just pass a pointer to the table
+		EffectBounds.Init();
+
+		for (TSharedPtr<FNiagaraSimulation>&it : Emitters)
+		{
+			it->SetConstants(Constants);
+			it->GetConstants().Merge(it->GetProperties()->ExternalConstants);
+			it->Tick(DeltaSeconds);
+
+			EffectBounds += it->GetEffectRenderer()->GetBounds();
+		}
+	}
+
+	NIAGARA_API void RenderModuleupdate();
+
+	FNiagaraSimulation *GetEmitter(uint32 idx)
+	{
+		return Emitters[idx].Get();
+	}
+
+
+	UNiagaraComponent *GetComponent() { return Component; }
+
+	TArray< TSharedPtr<FNiagaraSimulation> > &GetEmitters()	{ return Emitters; }
+
+	FBox GetEffectBounds()	{ return EffectBounds;  }
+private:
+	UNiagaraComponent *Component;
+	FBox EffectBounds;
+
+	/** Local constant table. */
+	FNiagaraConstantMap Constants;
+
+	TArray< TSharedPtr<FNiagaraSimulation> > Emitters;
+
+	FNiagaraSparseVolumeDataObject VolumeGrid;
+
+};
\ No newline at end of file
diff --git a/Engine/Source/Runtime/Niagara/Classes/NiagaraScript.h b/Engine/Source/Runtime/Niagara/Classes/NiagaraScript.h
new file mode 100644
--- /dev/null
+++ b/Engine/Source/Runtime/Niagara/Classes/NiagaraScript.h
@@ -0,0 +1,48 @@
+// Copyright 1998-2015 Epic Games, Inc. All Rights Reserved.
+
+#pragma once
+
+#include "NiagaraScriptConstantData.h"
+
+#include "NiagaraScript.generated.h"
+
+#define NIAGARA_INVALID_MEMORY (0xBA)
+
+/** Defines what will happen to unused attributes when a script is run. */
+UENUM()
+enum class EUnusedAttributeBehaviour : uint8
+{
+	/** The previous value of the attribute is copied across. */
+	Copy,
+	/** The attribute is set to zero. */
+	Zero,
+	/** The attribute is untouched. */
+	None,
+	/** The memory for the attribute is set to NIAGARA_INVALID_MEMORY. */
+	MarkInvalid, 
+};
+
+/** Runtime script for a Niagara system */
+UCLASS(MinimalAPI)
+class UNiagaraScript : public UObject
+{
+	GENERATED_UCLASS_BODY()
+
+	/** Byte code to execute for this system */
+	UPROPERTY()
+	TArray<uint8> ByteCode;
+
+	/** All the data for using constants in the script. */
+	UPROPERTY()
+	FNiagaraScriptConstantData ConstantData;
+
+	/** Attributes used by this script. */
+	UPROPERTY()
+ 	TArray<FNiagaraVariableInfo> Attributes;
+
+#if WITH_EDITORONLY_DATA
+	/** 'Source' data/graphs for this script */
+	UPROPERTY()
+	class UNiagaraScriptSourceBase*	Source;
+#endif
+};
\ No newline at end of file
diff --git a/Engine/Source/Runtime/Niagara/Classes/NiagaraScriptSourceBase.h b/Engine/Source/Runtime/Niagara/Classes/NiagaraScriptSourceBase.h
new file mode 100644
--- /dev/null
+++ b/Engine/Source/Runtime/Niagara/Classes/NiagaraScriptSourceBase.h
@@ -0,0 +1,28 @@
+// Copyright 1998-2015 Epic Games, Inc. All Rights Reserved.
+
+#pragma once
+#include "NiagaraScriptSourceBase.generated.h"
+
+
+struct EditorExposedVectorConstant
+{
+	FName ConstName;
+	FVector4 Value;
+};
+
+struct EditorExposedVectorCurveConstant
+{
+	FName ConstName;
+	class UCurveVector *Value;
+};
+
+/** Runtime data for a Niagara system */
+UCLASS(MinimalAPI)
+class UNiagaraScriptSourceBase : public UObject
+{
+	GENERATED_UCLASS_BODY()
+
+	TArray<TSharedPtr<EditorExposedVectorConstant> > ExposedVectorConstants;
+	TArray<TSharedPtr<EditorExposedVectorCurveConstant> > ExposedVectorCurveConstants;
+	virtual void Compile() {};
+};
\ No newline at end of file
diff --git a/Engine/Source/Runtime/Niagara/Classes/NiagaraSimulation.h b/Engine/Source/Runtime/Niagara/Classes/NiagaraSimulation.h
new file mode 100644
--- /dev/null
+++ b/Engine/Source/Runtime/Niagara/Classes/NiagaraSimulation.h
@@ -0,0 +1,167 @@
+// Copyright 1998-2015 Epic Games, Inc. All Rights Reserved.
+
+/*==============================================================================
+NiagaraSimulation.h: Niagara emitter simulation class
+==============================================================================*/
+#pragma once
+
+#include "NiagaraScript.h"
+#include "NiagaraComponent.h"
+#include "NiagaraSimulation.generated.h"
+
+DECLARE_CYCLE_STAT(TEXT("Tick"), STAT_NiagaraTick, STATGROUP_Niagara);
+DECLARE_CYCLE_STAT(TEXT("Simulate"), STAT_NiagaraSimulate, STATGROUP_Niagara);
+DECLARE_CYCLE_STAT(TEXT("Spawn"), STAT_NiagaraSpawn, STATGROUP_Niagara);
+DECLARE_CYCLE_STAT(TEXT("Kill"), STAT_NiagaraKill, STATGROUP_Niagara);
+
+
+class NiagaraEffectRenderer;
+
+UENUM()
+enum EEmitterRenderModuleType
+{
+	RMT_None = 0,
+	RMT_Sprites,
+	RMT_Ribbon,
+	RMT_Trails,
+	RMT_Meshes
+};
+
+
+
+/** 
+ *	FniagaraEmitterProperties stores the attributes of an FNiagaraSimulation
+ *	that need to be serialized and are used for its initialization 
+ */
+USTRUCT()
+struct FNiagaraEmitterProperties
+{
+	GENERATED_USTRUCT_BODY()
+public:
+	FNiagaraEmitterProperties() : 
+		bIsEnabled(true), SpawnRate(50), UpdateScript(nullptr), SpawnScript(nullptr), Material(nullptr), 
+		RenderModuleType(RMT_Sprites), RendererProperties(nullptr)
+	{
+		Name = FString(TEXT("New Emitter"));
+	}
+	UPROPERTY(EditAnywhere, Category = "Emitter Properties")
+	FString Name;
+	UPROPERTY(EditAnywhere, Category = "Emitter Properties")
+	bool bIsEnabled;
+	UPROPERTY(EditAnywhere, Category = "Emitter Properties")
+	float SpawnRate;
+	UPROPERTY(EditAnywhere, Category = "Emitter Properties")
+	UNiagaraScript *UpdateScript;
+	UPROPERTY(EditAnywhere, Category = "Emitter Properties")
+	UNiagaraScript *SpawnScript;
+	UPROPERTY(EditAnywhere, Category = "Emitter Properties")
+	UMaterial *Material;
+	UPROPERTY(EditAnywhere, Category = "Emitter Properties")
+	TEnumAsByte<EEmitterRenderModuleType> RenderModuleType;
+
+	UPROPERTY()
+	class UNiagaraEffectRendererProperties *RendererProperties;
+
+	UPROPERTY(EditAnywhere, Category = "Emitter Properties")
+	FNiagaraConstantMap ExternalConstants;		// these are the update script constants from the effect editor; will be added to the emitter's constant map
+
+	UPROPERTY(EditAnywhere, Category = "Emitter Properties")
+	FNiagaraConstantMap ExternalSpawnConstants;		// these are the spawn script constants from the effect editor; will be added to the emitter's constant map
+};
+
+
+
+/**
+* A niagara particle simulation.
+*/
+struct FNiagaraSimulation
+{
+public:
+	explicit FNiagaraSimulation(FNiagaraEmitterProperties *InProps);
+	FNiagaraSimulation(FNiagaraEmitterProperties *Props, ERHIFeatureLevel::Type InFeatureLevel);
+	virtual ~FNiagaraSimulation()
+	{}
+
+	void Init();
+
+	void Tick(float DeltaSeconds);
+
+
+	FBox GetBounds() const { return CachedBounds; }
+
+
+	FNiagaraEmitterParticleData &GetData()	{ return Data; }
+
+	void SetConstants(const FNiagaraConstantMap &InMap)	{ Constants = InMap; }
+	FNiagaraConstantMap &GetConstants()	{ return Constants; }
+
+	NiagaraEffectRenderer *GetEffectRenderer()	{ return EffectRenderer; }
+	
+	bool IsEnabled()const 	{ return bIsEnabled;  }
+
+	void NIAGARA_API SetRenderModuleType(EEmitterRenderModuleType Type, ERHIFeatureLevel::Type FeatureLevel);
+
+	int32 GetNumParticles()	{ return Data.GetNumParticles(); }
+
+	FNiagaraEmitterProperties *GetProperties()	{ return Props; }
+	void SetProperties(FNiagaraEmitterProperties *InProps);
+
+	float NIAGARA_API GetTotalCPUTime();
+	int	NIAGARA_API GetTotalBytesUsed();
+
+	void NIAGARA_API SetSpawnScript(UNiagaraScript* Script);
+	void NIAGARA_API SetUpdateScript(UNiagaraScript* Script);
+	void NIAGARA_API SetEnabled(bool bEnabled);
+private:
+	FNiagaraEmitterProperties *Props;		// points to an entry in the array of FNiagaraProperties stored in the EffectInstance (itself pointing to the effect's properties)
+
+	/* The age of the emitter*/
+	float Age;
+	/* If false, don't tick or render*/
+	bool bIsEnabled;
+	/* Seconds taken to process everything (including rendering) */
+	float CPUTimeMS;
+	
+	/** Local constant set. */
+	FNiagaraConstantMap Constants;
+	/** particle simulation data */
+	FNiagaraEmitterParticleData Data;
+	/** Keep partial particle spawns from last frame */
+	float SpawnRemainder;
+	/** The cached ComponentToWorld transform. */
+	FTransform CachedComponentToWorld;
+	/** Cached bounds. */
+	FBox CachedBounds;
+
+	NiagaraEffectRenderer *EffectRenderer;
+
+	/** Calc number to spawn */
+	int32 CalcNumToSpawn(float DeltaSeconds)
+	{
+		float FloatNumToSpawn = SpawnRemainder + (DeltaSeconds * Props->SpawnRate);
+		int32 NumToSpawn = FMath::FloorToInt(FloatNumToSpawn);
+		SpawnRemainder = FloatNumToSpawn - NumToSpawn;
+		return NumToSpawn;
+	}
+	
+	/** Runs a script in the VM over a specific range of particles. */
+	void RunVMScript(UNiagaraScript* Script, EUnusedAttributeBehaviour UnusedAttribBehaviour);
+	void RunVMScript(UNiagaraScript* Script, EUnusedAttributeBehaviour UnusedAttribBehaviour, uint32 StartParticle);
+	void RunVMScript(UNiagaraScript* Script, EUnusedAttributeBehaviour UnusedAttribBehaviour, uint32 StartParticle, uint32 NumParticles);
+
+	/** Util to move a particle */
+	void MoveParticleToIndex(int32 SrcIndex, int32 DestIndex)
+	{
+		FVector4 *SrcPtr = Data.GetCurrentBuffer() + SrcIndex;
+		FVector4 *DestPtr = Data.GetCurrentBuffer() + DestIndex;
+
+		for (int32 AttrIndex = 0; AttrIndex < Data.GetNumAttributes(); AttrIndex++)
+		{
+			*DestPtr = *SrcPtr;
+			DestPtr += Data.GetParticleAllocation();
+			SrcPtr += Data.GetParticleAllocation();
+		}
+	}
+
+	bool CheckAttriubtesForRenderer();
+};
diff --git a/Engine/Source/Runtime/Niagara/Niagara.Build.cs b/Engine/Source/Runtime/Niagara/Niagara.Build.cs
new file mode 100644
--- /dev/null
+++ b/Engine/Source/Runtime/Niagara/Niagara.Build.cs
@@ -0,0 +1,37 @@
+// Copyright 1998-2015 Epic Games, Inc. All Rights Reserved.
+
+using UnrealBuildTool;
+
+public class Niagara : ModuleRules
+{
+    public Niagara(TargetInfo Target)
+    {
+        PrivateDependencyModuleNames.AddRange(
+            new string[] {
+				"Core",
+				"CoreUObject",
+                "Engine",
+				"RenderCore",
+                "ShaderCore",
+                "VectorVM",
+                "RHI"
+			}
+        );
+
+
+        PublicDependencyModuleNames.AddRange(
+            new string[] {
+                "RHI"
+            }
+        );
+
+
+        PrivateIncludePaths.AddRange(
+            new string[] {
+                "Runtime/Niagara/Private",
+                "Runtime/Engine/Private",
+                "Runtime/Engine/Classes/Engine"
+            })
+        ;
+    }
+}
diff --git a/Engine/Source/Runtime/Niagara/NiagaraConstantSet.h b/Engine/Source/Runtime/Niagara/NiagaraConstantSet.h
new file mode 100644
--- /dev/null
+++ b/Engine/Source/Runtime/Niagara/NiagaraConstantSet.h
@@ -0,0 +1,384 @@
+// Copyright 1998-2015 Epic Games, Inc. All Rights Reserved.
+
+#pragma once
+
+#include "NiagaraCommon.h"
+#include "NiagaraConstantSet.generated.h"
+
+
+USTRUCT()
+struct FNiagaraConstantMap
+{
+	GENERATED_USTRUCT_BODY()
+
+private:
+	TMap<FNiagaraVariableInfo, float> ScalarConstants;
+	TMap<FNiagaraVariableInfo, FVector4> VectorConstants;
+	TMap<FNiagaraVariableInfo, FMatrix> MatrixConstants;
+	TMap<FNiagaraVariableInfo, class FNiagaraDataObject*> DataConstants;
+
+public:
+
+	void Empty()
+	{
+		ScalarConstants.Empty();
+		VectorConstants.Empty();
+		MatrixConstants.Empty();
+	}
+
+	void SetOrAdd(FNiagaraVariableInfo ID, float Sc)
+	{
+		check(ID.Type == ENiagaraDataType::Scalar);
+		ScalarConstants.FindOrAdd(ID) = Sc;
+	}
+
+	void SetOrAdd(FNiagaraVariableInfo ID, const FVector4& Vc)
+	{
+		check(ID.Type == ENiagaraDataType::Vector);
+		VectorConstants.FindOrAdd(ID) = Vc;
+	}
+
+	void SetOrAdd(FNiagaraVariableInfo ID, const FMatrix& Mc)
+	{
+		check(ID.Type == ENiagaraDataType::Matrix);
+		MatrixConstants.FindOrAdd(ID) = Mc;
+	}
+
+	void SetOrAdd(FNiagaraVariableInfo ID, FNiagaraDataObject* Cc)
+	{
+		check(ID.Type == ENiagaraDataType::Curve);
+		DataConstants.FindOrAdd(ID) = Cc;
+	}
+
+	void SetOrAdd(FName Name, float Sc)
+	{
+		ScalarConstants.FindOrAdd(FNiagaraVariableInfo(Name, ENiagaraDataType::Scalar)) = Sc;
+	}
+
+	void SetOrAdd(FName Name, const FVector4& Vc)
+	{
+		VectorConstants.FindOrAdd(FNiagaraVariableInfo(Name, ENiagaraDataType::Vector)) = Vc;
+	}
+
+	void SetOrAdd(FName Name, const FMatrix& Mc)
+	{
+		MatrixConstants.FindOrAdd(FNiagaraVariableInfo(Name, ENiagaraDataType::Matrix)) = Mc;
+	}
+
+	void SetOrAdd(FName Name, class FNiagaraDataObject* Cc)
+	{
+		DataConstants.FindOrAdd(FNiagaraVariableInfo(Name, ENiagaraDataType::Curve)) = Cc;
+	}
+
+
+	float* FindScalar(FNiagaraVariableInfo ID){ return ScalarConstants.Find(ID); }
+	FVector4* FindVector(FNiagaraVariableInfo ID){ return VectorConstants.Find(ID); }
+	FMatrix* FindMatrix(FNiagaraVariableInfo ID){ return MatrixConstants.Find(ID); }
+	const float* FindScalar(FNiagaraVariableInfo ID)const { return ScalarConstants.Find(ID); }
+	const FVector4* FindVector(FNiagaraVariableInfo ID)const { return VectorConstants.Find(ID); }
+	const FMatrix* FindMatrix(FNiagaraVariableInfo ID)const { return MatrixConstants.Find(ID); }
+	FNiagaraDataObject* FindDataObj(FNiagaraVariableInfo ID) const
+	{
+		FNiagaraDataObject * const *Ptr = DataConstants.Find(ID);
+		if (Ptr)
+		{
+			return *Ptr;
+		}
+
+		return nullptr; 
+	}
+	float* FindScalar(FName Name) { return ScalarConstants.Find(FNiagaraVariableInfo(Name,ENiagaraDataType::Scalar)); }
+	FVector4* FindVector(FName Name) { return VectorConstants.Find(FNiagaraVariableInfo(Name, ENiagaraDataType::Vector)); }
+	FMatrix* FindMatrix(FName Name) { return MatrixConstants.Find(FNiagaraVariableInfo(Name, ENiagaraDataType::Matrix)); }
+	FNiagaraDataObject* FindDataObj(FName Name)const { return FindDataObj(FNiagaraVariableInfo(Name, ENiagaraDataType::Curve)); }
+
+	void Merge(FNiagaraConstantMap &InMap)
+	{
+		for (auto MapIt = InMap.ScalarConstants.CreateIterator(); MapIt; ++MapIt)
+		{
+			SetOrAdd(MapIt.Key(), MapIt.Value());
+		}
+
+		for (auto MapIt = InMap.VectorConstants.CreateIterator(); MapIt; ++MapIt)
+		{
+			SetOrAdd(MapIt.Key(), MapIt.Value());
+		}
+
+		for (auto MapIt = InMap.MatrixConstants.CreateIterator(); MapIt; ++MapIt)
+		{
+			SetOrAdd(MapIt.Key(), MapIt.Value());
+		}
+
+		for (auto MapIt = InMap.DataConstants.CreateIterator(); MapIt; ++MapIt)
+		{
+			SetOrAdd(MapIt.Key(), MapIt.Value());
+		}
+	}
+
+	const TMap<FNiagaraVariableInfo, class FNiagaraDataObject*> &GetDataConstants()	const
+	{ 
+		return DataConstants; 
+	}
+
+
+	virtual bool Serialize(FArchive &Ar)
+	{
+		// TODO: can't serialize the data object constants at the moment; need to figure that out
+		Ar << ScalarConstants << VectorConstants << MatrixConstants;
+		return true;
+	}
+};
+
+
+
+
+USTRUCT()
+struct FNiagaraConstants
+{
+	GENERATED_USTRUCT_BODY()
+
+private:
+
+	UPROPERTY()
+	TArray<FNiagaraVariableInfo> ScalarConstantsInfo;
+	UPROPERTY()
+	TArray<FNiagaraVariableInfo> VectorConstantsInfo;
+	UPROPERTY()
+	TArray<FNiagaraVariableInfo> MatrixConstantsInfo;
+	//UPROPERTY()
+	TArray<FNiagaraVariableInfo> DataObjectConstantsInfo;
+
+	UPROPERTY()
+	TArray<float> ScalarConstants;
+	UPROPERTY()
+	TArray<FVector4> VectorConstants;
+	UPROPERTY()
+	TArray<FMatrix> MatrixConstants;
+	//UPROPERTY()
+	TArray<FNiagaraDataObject*> DataObjectConstants;
+
+public:
+
+	void Empty()
+	{
+		ScalarConstantsInfo.Empty();
+		VectorConstantsInfo.Empty();
+		MatrixConstantsInfo.Empty();
+		DataObjectConstantsInfo.Empty();
+		ScalarConstants.Empty();
+		VectorConstants.Empty();
+		MatrixConstants.Empty();
+		DataObjectConstants.Empty();
+	}
+
+	void GetScalarConstant(int32 Index, float& OutValue, FNiagaraVariableInfo& OutInfo)const
+	{
+		OutInfo = ScalarConstantsInfo[Index]; 
+		OutValue = ScalarConstants[Index];
+	}
+	void GetVectorConstant(int32 Index, FVector4& OutValue, FNiagaraVariableInfo& OutInfo)const
+	{
+		OutInfo = VectorConstantsInfo[Index];
+		OutValue = VectorConstants[Index];
+	}
+	void GetMatrixConstant(int32 Index, FMatrix& OutValue, FNiagaraVariableInfo& OutInfo)const
+	{
+		OutInfo = MatrixConstantsInfo[Index];
+		OutValue = MatrixConstants[Index];
+	}
+	void GetDataObjectConstant(int32 Index, FNiagaraDataObject*& OutValue, FNiagaraVariableInfo& OutInfo)const
+	{
+		OutInfo = DataObjectConstantsInfo[Index];
+		OutValue = DataObjectConstants[Index];
+	}
+
+	const int32 GetNumScalarConstants()const{ return ScalarConstantsInfo.Num(); }
+	const int32 GetNumVectorConstants()const{ return VectorConstantsInfo.Num(); }
+	const int32 GetNumMatrixConstants()const{ return MatrixConstantsInfo.Num(); }
+	const int32 GetNumDataObjectConstants()const{ return DataObjectConstantsInfo.Num(); }
+
+	/** Fills the entire constants set into the constant table. */
+	void AppendToConstantsTable(TArray<FVector4>& ConstantsTable)const
+	{
+		int32 Idx = ConstantsTable.Num();
+		int32 NewIdx = 0;
+		ConstantsTable.AddUninitialized(ScalarTableSize() + VectorTableSize() + MatrixTableSize());
+		for (float Sc : ScalarConstants)
+		{
+			ConstantsTable[Idx + (NewIdx / 4)][NewIdx % 4] = Sc;
+			++NewIdx;
+			Idx = NewIdx % 4 == 0 ? Idx + 1 : Idx;
+		}
+
+		Idx += FMath::Min(1, NewIdx % 4);//Move to the next table entry if needed.
+		for (const FVector4& Vc : VectorConstants)
+		{
+			ConstantsTable[Idx++] = Vc;
+		}
+		for (const FMatrix& Mc : MatrixConstants)
+		{
+			ConstantsTable[Idx] = FVector4(Mc.M[0][0], Mc.M[0][1], Mc.M[0][2], Mc.M[0][3]);
+			ConstantsTable[Idx + 1] = FVector4(Mc.M[1][0], Mc.M[1][1], Mc.M[1][2], Mc.M[1][3]);
+			ConstantsTable[Idx + 2] = FVector4(Mc.M[2][0], Mc.M[2][1], Mc.M[2][2], Mc.M[2][3]);
+			ConstantsTable[Idx + 3] = FVector4(Mc.M[3][0], Mc.M[3][1], Mc.M[3][2], Mc.M[3][3]);
+			Idx += 4;
+		}
+	}
+
+	void AppendBufferConstants(TArray<class FNiagaraDataObject*> &Objs) const
+	{
+		Objs.Append(DataObjectConstants);
+	}
+
+	/** 
+	Fills only selected constants into the table. In the order they appear in the array of passed names not the order they appear in the set. 
+	Checks the passed map for entires to superceed the default values in the set.
+	*/
+	void AppendToConstantsTable(TArray<FVector4>& ConstantsTable, const FNiagaraConstantMap& Externals)const
+	{
+		int32 Idx = ConstantsTable.Num();
+		int32 NewIdx = 0;
+		ConstantsTable.AddUninitialized(ScalarTableSize() + VectorTableSize() + MatrixTableSize());
+		for (int32 i = 0; i < ScalarConstants.Num(); ++i)
+		{
+			const float* Scalar = Externals.FindScalar(ScalarConstantsInfo[i]);
+			ConstantsTable[Idx + (NewIdx / 4)][NewIdx % 4] = Scalar ? *Scalar : ScalarConstants[i];
+			++NewIdx;
+			Idx = NewIdx % 4 == 0 ? Idx + 1 : Idx;
+		}
+
+		Idx += FMath::Min(1, NewIdx % 4);//Move to the next table entry if needed.
+
+		for (int32 i = 0; i < VectorConstants.Num(); ++i)
+		{
+			const FVector4* Vector = Externals.FindVector(VectorConstantsInfo[i]);
+			ConstantsTable[Idx++] = Vector ? *Vector : VectorConstants[i];
+		}
+
+		for (int32 i = 0; i < MatrixConstants.Num(); ++i)
+		{
+			const FMatrix* Mat = Externals.FindMatrix(MatrixConstantsInfo[i]);
+			FMatrix Mc = Mat ? *Mat : MatrixConstants[i];
+			ConstantsTable[Idx] = FVector4(Mc.M[0][0], Mc.M[0][1], Mc.M[0][2], Mc.M[0][3]);
+			ConstantsTable[Idx + 1] = FVector4(Mc.M[1][0], Mc.M[1][1], Mc.M[1][2], Mc.M[1][3]);
+			ConstantsTable[Idx + 2] = FVector4(Mc.M[2][0], Mc.M[2][1], Mc.M[2][2], Mc.M[2][3]);
+			ConstantsTable[Idx + 3] = FVector4(Mc.M[3][0], Mc.M[3][1], Mc.M[3][2], Mc.M[3][3]);
+		}
+	}
+
+	void AppendExternalBufferConstants(TArray<class FNiagaraDataObject*> &DataObjs, const FNiagaraConstantMap& Externals) const
+	{
+		for (int32 i = 0; i < DataObjectConstantsInfo.Num(); ++i)
+		{
+			FNiagaraVariableInfo CcID = DataObjectConstantsInfo[i];
+			FNiagaraDataObject* Data = Externals.FindDataObj(CcID);
+			DataObjs.Add(Data);
+		}
+	}
+
+	void SetOrAdd(const FNiagaraVariableInfo& Constant, float Value)
+	{
+		int32 Idx = ScalarConstantsInfo.Find(Constant);
+		if (Idx == INDEX_NONE)
+		{
+			ScalarConstantsInfo.Add(Constant);
+			ScalarConstants.Add(Value);
+		}
+		else
+		{
+			ScalarConstants[Idx] = Value;
+		}
+	}
+
+	void SetOrAdd(const FNiagaraVariableInfo& Constant, const FVector4& Value)
+	{
+		int32 Idx = VectorConstantsInfo.Find(Constant);
+		if (Idx == INDEX_NONE)
+		{
+			VectorConstantsInfo.Add(Constant);
+			VectorConstants.Add(Value);
+		}
+		else
+		{
+			VectorConstants[Idx] = Value;
+		}
+	}
+
+	void SetOrAdd(const FNiagaraVariableInfo& Constant, const FMatrix& Value)
+	{
+		int32 Idx = MatrixConstantsInfo.Find(Constant);
+		if (Idx == INDEX_NONE)
+		{
+			MatrixConstantsInfo.Add(Constant);
+			MatrixConstants.Add(Value);
+		}
+		else
+		{
+			MatrixConstants[Idx] = Value;
+		}
+	}
+
+	void SetOrAdd(const FNiagaraVariableInfo& Constant, FNiagaraDataObject* Value)
+	{
+		int32 Idx = DataObjectConstantsInfo.Find(Constant);
+		if (Idx == INDEX_NONE)
+		{
+			DataObjectConstantsInfo.Add(Constant);
+			DataObjectConstants.Add(Value);
+		}
+		else
+		{
+			DataObjectConstants[Idx] = Value;
+		}
+	}
+
+	FORCEINLINE int32 NumScalars()const { return ScalarConstants.Num(); }
+	FORCEINLINE int32 NumVectors()const { return VectorConstants.Num(); }
+	FORCEINLINE int32 NumMatrices()const { return MatrixConstants.Num(); }
+	FORCEINLINE int32 ScalarTableSize()const { return(NumScalars() / 4) + FMath::Min(1,NumScalars() % 4); }
+	FORCEINLINE int32 VectorTableSize()const { return(NumVectors()); }
+	FORCEINLINE int32 MatrixTableSize()const { return(NumMatrices() * 4); }
+
+	/**	Returns the number of vectors these constants would use in a constants table. */
+	FORCEINLINE int32 GetTableSize()const { return ScalarTableSize() + VectorTableSize() + MatrixTableSize(); }
+
+	/** Return the first constant used for this constant in a table. */
+	FORCEINLINE int32 GetTableIndex_Scalar(const FNiagaraVariableInfo& Constant)const
+	{ 
+		int Idx = ScalarConstantsInfo.Find(Constant);
+		if (Idx != INDEX_NONE)
+		{
+			return Idx / 4;
+		}
+		return INDEX_NONE;
+	}
+
+	FORCEINLINE int32 GetTableIndex_Vector(const FNiagaraVariableInfo& Constant)const
+	{ 
+		int Idx = VectorConstantsInfo.Find(Constant);
+		if (Idx != INDEX_NONE)
+		{
+			return ScalarTableSize() + Idx;
+		}
+		return INDEX_NONE;
+	}
+
+	FORCEINLINE int32 GetTableIndex_Matrix(const FNiagaraVariableInfo& Constant)const
+	{ 
+		int Idx = MatrixConstantsInfo.Find(Constant);
+		if (Idx != INDEX_NONE)
+		{
+			return ScalarTableSize() + VectorTableSize() + Idx * 4;
+		}
+		return INDEX_NONE;
+	}
+
+	FORCEINLINE int32 GetTableIndex_DataObj(const FNiagaraVariableInfo& Constant)const
+	{ 
+		return DataObjectConstantsInfo.Find(Constant); 
+	}
+
+	/** Return the absolute index of the passed scalar. Can be used to get the component index also. */
+	FORCEINLINE int32 GetAbsoluteIndex_Scalar(const FNiagaraVariableInfo& Constant)const{ return ScalarConstantsInfo.Find(Constant); }
+};
+
diff --git a/Engine/Source/Runtime/Niagara/NiagaraScriptConstantData.h b/Engine/Source/Runtime/Niagara/NiagaraScriptConstantData.h
new file mode 100644
--- /dev/null
+++ b/Engine/Source/Runtime/Niagara/NiagaraScriptConstantData.h
@@ -0,0 +1,94 @@
+// Copyright 1998-2015 Epic Games, Inc. All Rights Reserved.
+
+#pragma once
+
+#include "NiagaraConstantSet.h"
+
+#include "NiagaraScriptConstantData.generated.h"
+
+USTRUCT()
+struct FNiagaraScriptConstantData
+{
+	GENERATED_USTRUCT_BODY()
+
+	/** The set of external constants for this script. */
+	UPROPERTY()
+	FNiagaraConstants ExternalConstants;
+
+	/** The set of internal constants for this script. */
+	UPROPERTY()
+	FNiagaraConstants InternalConstants;
+
+	void Empty()
+	{
+		ExternalConstants.Empty();
+		InternalConstants.Empty();
+	}
+
+	/** Fill a constants table ready for use in an update or spawn. */
+	void FillConstantTable(const FNiagaraConstantMap& ExternalConstantsMap, TArray<FVector4>& OutConstantTable, TArray<FNiagaraDataObject *> &DataObjTable) const
+	{
+		//First up in the table comes the External constants in scalar->vector->matrix order.
+		//Only fills the constants actually used by the script.
+		OutConstantTable.Empty();
+		ExternalConstants.AppendToConstantsTable(OutConstantTable, ExternalConstantsMap);
+		ExternalConstants.AppendExternalBufferConstants(DataObjTable, ExternalConstantsMap);
+		//Next up add all the internal constants from the script.
+		InternalConstants.AppendToConstantsTable(OutConstantTable);
+		InternalConstants.AppendBufferConstants(DataObjTable);
+	}
+
+	template<typename T>
+	void SetOrAddInternal(const FNiagaraVariableInfo& Constant, const T& Value)
+	{
+		InternalConstants.SetOrAdd(Constant, Value);
+	}
+
+	template<typename T>
+	void SetOrAddExternal(const FNiagaraVariableInfo& Constant, const T& Value)
+	{
+		ExternalConstants.SetOrAdd(Constant, Value);
+	}
+
+	/**
+	Calculates the index into a a constant table created from this constant data for the given constant name.
+	Obviously, assumes the constant data is complete. Any additions etc will invalidate previously calculated indexes.
+	*/
+	void GetTableIndex(const FNiagaraVariableInfo& Constant, bool bInternal, int32& OutConstantIdx, int32& OutComponentIndex, ENiagaraDataType& OutType)const
+	{
+		int32 Base = bInternal ? ExternalConstants.GetTableSize() : 0;
+		int32 ConstIdx = INDEX_NONE;
+		int32 CompIdx = INDEX_NONE;
+		ENiagaraDataType Type;
+		const FNiagaraConstants& Consts = bInternal ? InternalConstants : ExternalConstants;
+
+		ConstIdx = Consts.GetAbsoluteIndex_Scalar(Constant);
+		Type = ENiagaraDataType::Scalar;
+		if (ConstIdx == INDEX_NONE)
+		{
+			ConstIdx = Consts.GetTableIndex_Vector(Constant);
+			Type = ENiagaraDataType::Vector;
+			if (ConstIdx == INDEX_NONE)
+			{
+				ConstIdx = Consts.GetTableIndex_Matrix(Constant);
+				Type = ENiagaraDataType::Matrix;
+
+				if (ConstIdx == INDEX_NONE)	// curves/buffers are in a separate table, so set base to 0
+				{
+					ConstIdx = Consts.GetTableIndex_DataObj(Constant);
+					Type = ENiagaraDataType::Curve;
+					Base = 0;
+				}
+			}
+		}
+		else
+		{
+			CompIdx = ConstIdx % 4;
+			ConstIdx /= 4;
+		}
+
+		OutConstantIdx = Base + ConstIdx;
+		OutComponentIndex = CompIdx;
+		OutType = Type;
+	}
+};
diff --git a/Engine/Source/Runtime/Niagara/Private/NiagaraActor.cpp b/Engine/Source/Runtime/Niagara/Private/NiagaraActor.cpp
new file mode 100644
--- /dev/null
+++ b/Engine/Source/Runtime/Niagara/Private/NiagaraActor.cpp
@@ -0,0 +1,86 @@
+// Copyright 1998-2015 Epic Games, Inc. All Rights Reserved.
+
+
+#include "NiagaraPrivate.h"
+#include "NiagaraActor.h"
+#include "NiagaraComponent.h"
+#include "NiagaraEffect.h"
+
+ANiagaraActor::ANiagaraActor(const FObjectInitializer& ObjectInitializer)
+	: Super(ObjectInitializer)
+{
+	NiagaraComponent = CreateDefaultSubobject<UNiagaraComponent>(TEXT("NiagaraComponent0"));
+
+	RootComponent = NiagaraComponent;
+
+#if WITH_EDITORONLY_DATA
+	SpriteComponent = CreateEditorOnlyDefaultSubobject<UBillboardComponent>(TEXT("Sprite"));
+	ArrowComponent = CreateEditorOnlyDefaultSubobject<UArrowComponent>(TEXT("ArrowComponent0"));
+
+	if (!IsRunningCommandlet())
+	{
+		// Structure to hold one-time initialization
+		struct FConstructorStatics
+		{
+			ConstructorHelpers::FObjectFinderOptional<UTexture2D> SpriteTextureObject;
+			FName ID_Effects;
+			FText NAME_Effects;
+			FConstructorStatics()
+				: SpriteTextureObject(TEXT("/Engine/EditorResources/S_Emitter"))
+				, ID_Effects(TEXT("Effects"))
+				, NAME_Effects(NSLOCTEXT("SpriteCategory", "Effects", "Effects"))
+			{
+			}
+		};
+		static FConstructorStatics ConstructorStatics;
+
+		if (SpriteComponent)
+		{
+			SpriteComponent->Sprite = ConstructorStatics.SpriteTextureObject.Get();
+			SpriteComponent->RelativeScale3D = FVector(0.5f, 0.5f, 0.5f);
+			SpriteComponent->bHiddenInGame = true;
+			SpriteComponent->bIsScreenSizeScaled = true;
+			SpriteComponent->SpriteInfo.Category = ConstructorStatics.ID_Effects;
+			SpriteComponent->SpriteInfo.DisplayName = ConstructorStatics.NAME_Effects;
+			SpriteComponent->AttachParent = NiagaraComponent;
+			SpriteComponent->bReceivesDecals = false;
+		}
+
+		if (ArrowComponent)
+		{
+			ArrowComponent->ArrowColor = FColor(0, 255, 128);
+
+			ArrowComponent->ArrowSize = 1.5f;
+			ArrowComponent->bTreatAsASprite = true;
+			ArrowComponent->SpriteInfo.Category = ConstructorStatics.ID_Effects;
+			ArrowComponent->SpriteInfo.DisplayName = ConstructorStatics.NAME_Effects;
+			ArrowComponent->AttachParent = NiagaraComponent;
+			ArrowComponent->bIsScreenSizeScaled = true;
+		}
+	}
+#endif // WITH_EDITORONLY_DATA
+}
+
+/** Returns NiagaraComponent subobject **/
+UNiagaraComponent* ANiagaraActor::GetNiagaraComponent() const { return NiagaraComponent; }
+#if WITH_EDITORONLY_DATA
+/** Returns SpriteComponent subobject **/
+UBillboardComponent* ANiagaraActor::GetSpriteComponent() const { return SpriteComponent; }
+/** Returns ArrowComponent subobject **/
+UArrowComponent* ANiagaraActor::GetArrowComponent() const { return ArrowComponent; }
+#endif
+
+#if WITH_EDITOR
+bool ANiagaraActor::GetReferencedContentObjects(TArray<UObject*>& Objects) const
+{
+	Super::GetReferencedContentObjects(Objects);
+
+	if (UNiagaraEffect* Effect = NiagaraComponent->GetAsset())
+	{
+		Objects.Add(Effect);
+	}
+
+	return true;
+}
+#endif // WITH_EDITOR
+
diff --git a/Engine/Source/Runtime/Niagara/Private/NiagaraComponent.cpp b/Engine/Source/Runtime/Niagara/Private/NiagaraComponent.cpp
new file mode 100644
--- /dev/null
+++ b/Engine/Source/Runtime/Niagara/Private/NiagaraComponent.cpp
@@ -0,0 +1,311 @@
+// Copyright 1998-2015 Epic Games, Inc. All Rights Reserved.
+
+#include "NiagaraPrivate.h"
+#include "NiagaraComponent.h"
+#include "NiagaraScript.h"
+#include "VectorVM.h"
+#include "ParticleHelper.h"
+#include "Particles/ParticleResources.h"
+#include "NiagaraEffectRenderer.h"
+#include "NiagaraEffect.h"
+#include "NiagaraSimulation.h"
+#include "MeshBatch.h"
+#include "SceneUtils.h"
+#include "ComponentReregisterContext.h"
+
+DECLARE_CYCLE_STAT(TEXT("Gen Verts"),STAT_NiagaraGenerateVertices,STATGROUP_Niagara);
+DECLARE_DWORD_COUNTER_STAT(TEXT("NumParticles"),STAT_NiagaraNumParticles,STATGROUP_Niagara);
+
+DEFINE_LOG_CATEGORY(LogNiagara);
+
+FNiagaraSceneProxy::FNiagaraSceneProxy(const UNiagaraComponent* InComponent)
+		:	FPrimitiveSceneProxy(InComponent)
+{
+	UpdateEffectRenderers(InComponent->GetEffectInstance());
+}
+
+void FNiagaraSceneProxy::UpdateEffectRenderers(FNiagaraEffectInstance* InEffect)
+{
+	EffectRenderers.Empty();
+	if (InEffect)
+	{
+		for (TSharedPtr<FNiagaraSimulation>Emitter : InEffect->GetEmitters())
+		{
+			AddEffectRenderer(Emitter->GetEffectRenderer());
+		}
+	}
+}
+
+FNiagaraSceneProxy::~FNiagaraSceneProxy()
+{
+	ReleaseRenderThreadResources();
+}
+
+/** Called on render thread to assign new dynamic data */
+void FNiagaraSceneProxy::SetDynamicData_RenderThread(FNiagaraDynamicDataBase* NewDynamicData)
+{
+	for (NiagaraEffectRenderer* Renderer : EffectRenderers)
+	{
+		if (Renderer)
+		{
+			Renderer->SetDynamicData_RenderThread(NewDynamicData);
+		}
+	}
+	return;
+}
+
+
+void FNiagaraSceneProxy::ReleaseRenderThreadResources()
+{
+	for (NiagaraEffectRenderer* Renderer : EffectRenderers)
+	{
+		if (Renderer)
+		{
+			Renderer->ReleaseRenderThreadResources();
+		}
+	}
+	return;
+}
+
+// FPrimitiveSceneProxy interface.
+void FNiagaraSceneProxy::CreateRenderThreadResources()
+{
+	for (NiagaraEffectRenderer* Renderer : EffectRenderers)
+	{
+		if (Renderer)
+		{
+			Renderer->CreateRenderThreadResources();
+		}
+	}
+	return;
+}
+
+void FNiagaraSceneProxy::OnActorPositionChanged()
+{
+	//WorldSpacePrimitiveUniformBuffer.ReleaseResource();
+}
+
+void FNiagaraSceneProxy::OnTransformChanged()
+{
+	//WorldSpacePrimitiveUniformBuffer.ReleaseResource();
+}
+
+FPrimitiveViewRelevance FNiagaraSceneProxy::GetViewRelevance(const FSceneView* View)
+{
+	FPrimitiveViewRelevance Relevance;
+	Relevance.bDynamicRelevance = true;
+
+	for (NiagaraEffectRenderer* Renderer : EffectRenderers)
+	{
+		if (Renderer)
+		{
+			Relevance |= Renderer->GetViewRelevance(View, this);
+		}
+	}
+	return Relevance;
+}
+
+
+uint32 FNiagaraSceneProxy::GetMemoryFootprint() const
+{ 
+	return (sizeof(*this) + GetAllocatedSize()); 
+}
+
+uint32 FNiagaraSceneProxy::GetAllocatedSize() const
+{ 
+	uint32 DynamicDataSize = 0;
+	for (NiagaraEffectRenderer* Renderer : EffectRenderers)
+	{
+		if (Renderer)
+		{
+			DynamicDataSize += Renderer->GetDynamicDataSize();
+		}
+	}
+	return FPrimitiveSceneProxy::GetAllocatedSize() + DynamicDataSize;
+}
+
+
+void FNiagaraSceneProxy::GetDynamicMeshElements(const TArray<const FSceneView*>& Views, const FSceneViewFamily& ViewFamily, uint32 VisibilityMap, FMeshElementCollector& Collector) const
+{
+	for (NiagaraEffectRenderer* Renderer : EffectRenderers)
+	{
+		if (Renderer)
+		{
+			Renderer->GetDynamicMeshElements(Views, ViewFamily, VisibilityMap, Collector, this);
+		}
+	}
+}
+
+//////////////////////////////////////////////////////////////////////////
+
+UNiagaraComponent::UNiagaraComponent(const FObjectInitializer& ObjectInitializer)
+	: Super(ObjectInitializer)
+{
+	PrimaryComponentTick.bCanEverTick = true;
+	bTickInEditor = true;
+}
+
+
+void UNiagaraComponent::TickComponent(float DeltaSeconds, enum ELevelTick TickType, FActorComponentTickFunction* ThisTickFunction)
+{
+//	EmitterAge += DeltaSeconds;
+
+	if (EffectInstance)
+	{ 
+		static FNiagaraVariableInfo Const_Zero(TEXT("ZERO"), ENiagaraDataType::Vector);
+		static FNiagaraVariableInfo Const_DeltaTime(TEXT("Delta Time"), ENiagaraDataType::Vector);
+		static FNiagaraVariableInfo Const_EmitterPos(TEXT("Emitter Position"), ENiagaraDataType::Vector);
+		static FNiagaraVariableInfo Const_EmitterAge(TEXT("Emitter Age"), ENiagaraDataType::Vector);
+		static FNiagaraVariableInfo Const_EmitterX(TEXT("Emitter X Axis"), ENiagaraDataType::Vector);
+		static FNiagaraVariableInfo Const_EmitterY(TEXT("Emitter Y Axis"), ENiagaraDataType::Vector);
+		static FNiagaraVariableInfo Const_EmitterZ(TEXT("Emitter Z Axis"), ENiagaraDataType::Vector);
+		static FNiagaraVariableInfo Const_EmitterTransform(TEXT("Emitter Transform"), ENiagaraDataType::Matrix);
+		//Todo, open this up to the UI and setting via code and BPs.
+		EffectInstance->SetConstant(Const_Zero, FVector4(0.0f, 0.0f, 0.0f, 0.0f));	// zero constant
+		EffectInstance->SetConstant(Const_DeltaTime, FVector4(DeltaSeconds, DeltaSeconds, DeltaSeconds, DeltaSeconds));
+		EffectInstance->SetConstant(Const_EmitterPos, FVector4(ComponentToWorld.GetTranslation()));
+		EffectInstance->SetConstant(Const_EmitterX, FVector4(ComponentToWorld.GetUnitAxis(EAxis::X)));
+		EffectInstance->SetConstant(Const_EmitterY, FVector4(ComponentToWorld.GetUnitAxis(EAxis::Y)));
+		EffectInstance->SetConstant(Const_EmitterZ, FVector4(ComponentToWorld.GetUnitAxis(EAxis::Z)));
+		EffectInstance->SetConstant(Const_EmitterTransform, ComponentToWorld.ToMatrixWithScale());
+
+		EffectInstance->Tick(DeltaSeconds);
+	}
+
+	UpdateComponentToWorld();
+	MarkRenderDynamicDataDirty();
+}
+
+const UObject* UNiagaraComponent::AdditionalStatObject() const
+{
+	return Asset;
+}
+
+void UNiagaraComponent::OnRegister()
+{
+	Super::OnRegister();
+	if (Asset)
+	{
+		if (!EffectInstance)
+		{
+			EffectInstance = new FNiagaraEffectInstance(Asset, this);
+		}
+		{
+			EffectInstance->Init(this);
+
+			// initialize all render modules
+			ENQUEUE_UNIQUE_RENDER_COMMAND_TWOPARAMETER(
+				FChangeNiagaraRenderModule,
+				FNiagaraEffectInstance*, InEffect, this->EffectInstance,
+				UNiagaraComponent*, InComponent, this,
+				{
+					for (TSharedPtr<FNiagaraSimulation> Emitter : InEffect->GetEmitters())
+					{
+						Emitter->SetRenderModuleType(Emitter->GetProperties()->RenderModuleType, InComponent->GetWorld()->FeatureLevel);
+					}
+				}
+			);
+		}
+		VectorVM::Init();
+	}
+}
+
+
+void UNiagaraComponent::OnUnregister()
+{
+	Super::OnUnregister();
+}
+
+void UNiagaraComponent::SendRenderDynamicData_Concurrent()
+{
+	if (EffectInstance && SceneProxy)
+	{
+		FNiagaraSceneProxy* NiagaraProxy = static_cast<FNiagaraSceneProxy*>(SceneProxy);
+		for (int32 i = 0; i < EffectInstance->GetEmitters().Num(); i++)
+		{
+			TSharedPtr<FNiagaraSimulation> Emitter = EffectInstance->GetEmitters()[i];
+			NiagaraEffectRenderer* Renderer = Emitter->GetEffectRenderer();
+			if (Renderer)
+			{
+				FNiagaraDynamicDataBase* DynamicData = Renderer->GenerateVertexData(Emitter->GetData());
+
+				ENQUEUE_UNIQUE_RENDER_COMMAND_TWOPARAMETER(
+					FSendNiagaraDynamicData,
+					NiagaraEffectRenderer*, EffectRenderer, Emitter->GetEffectRenderer(),
+					FNiagaraDynamicDataBase*, DynamicData, DynamicData,
+					{
+					EffectRenderer->SetDynamicData_RenderThread(DynamicData);
+				});
+			}
+		}
+	}
+
+}
+
+int32 UNiagaraComponent::GetNumMaterials() const
+{
+	return 0;
+}
+
+
+FBoxSphereBounds UNiagaraComponent::CalcBounds(const FTransform& LocalToWorld) const
+{
+	FBox SimBounds(ForceInit);
+	if (EffectInstance)
+	{
+		FBoxSphereBounds BSBounds(EffectInstance->GetEffectBounds());
+		//BSBounds.Origin = LocalToWorld.GetLocation();
+		return BSBounds;
+	}
+	return FBoxSphereBounds();
+}
+
+FPrimitiveSceneProxy* UNiagaraComponent::CreateSceneProxy()
+{
+	FNiagaraSceneProxy* Proxy = new FNiagaraSceneProxy(this);
+	Proxy->UpdateEffectRenderers(EffectInstance);
+	return Proxy;
+}
+
+#if WITH_EDITOR
+void UNiagaraComponent::PostEditChangeProperty(FPropertyChangedEvent& PropertyChangedEvent)
+{
+	Super::PostEditChangeProperty(PropertyChangedEvent);
+
+	FComponentReregisterContext ReregisterContext(this);
+}
+#endif // WITH_EDITOR
+
+
+
+void UNiagaraComponent::SetAsset(UNiagaraEffect* InAsset)
+{
+	Asset = InAsset;
+
+	EffectInstance = new FNiagaraEffectInstance(Asset, this);
+}
+
+const TArray<FNiagaraVariableInfo>& UNiagaraComponent::GetSystemConstants()
+{
+	static TArray<FNiagaraVariableInfo> SystemConstants;
+	if (SystemConstants.Num() == 0)
+	{
+		static FNiagaraVariableInfo Const_Zero(TEXT("ZERO"), ENiagaraDataType::Vector);
+		static FNiagaraVariableInfo Const_DeltaTime(TEXT("Delta Time"), ENiagaraDataType::Vector);
+		static FNiagaraVariableInfo Const_EmitterPos(TEXT("Emitter Position"), ENiagaraDataType::Vector);
+		static FNiagaraVariableInfo Const_EmitterAge(TEXT("Emitter Age"), ENiagaraDataType::Vector);
+		static FNiagaraVariableInfo Const_EmitterX(TEXT("Emitter X Axis"), ENiagaraDataType::Vector);
+		static FNiagaraVariableInfo Const_EmitterY(TEXT("Emitter Y Axis"), ENiagaraDataType::Vector);
+		static FNiagaraVariableInfo Const_EmitterZ(TEXT("Emitter Z Axis"), ENiagaraDataType::Vector);
+		static FNiagaraVariableInfo Const_EmitterTransform(TEXT("Emitter Transform"), ENiagaraDataType::Matrix);
+		SystemConstants.Add(Const_Zero);
+		SystemConstants.Add(Const_DeltaTime);
+		SystemConstants.Add(Const_EmitterPos);
+		SystemConstants.Add(Const_EmitterAge);
+		SystemConstants.Add(Const_EmitterX);
+		SystemConstants.Add(Const_EmitterY);
+		SystemConstants.Add(Const_EmitterZ);
+		SystemConstants.Add(Const_EmitterTransform);
+	}
+	return SystemConstants;
+}
\ No newline at end of file
diff --git a/Engine/Source/Runtime/Niagara/Private/NiagaraConstantSet.cpp b/Engine/Source/Runtime/Niagara/Private/NiagaraConstantSet.cpp
new file mode 100644
--- /dev/null
+++ b/Engine/Source/Runtime/Niagara/Private/NiagaraConstantSet.cpp
@@ -0,0 +1,14 @@
+// Copyright 1998-2015 Epic Games, Inc. All Rights Reserved.
+
+#include "NiagaraPrivate.h"
+#include "NiagaraConstantSet.h"
+
+template<>
+struct TStructOpsTypeTraits<FNiagaraConstantMap> : public TStructOpsTypeTraitsBase
+{
+	enum
+	{
+		WithSerializer = true
+	};
+};
+IMPLEMENT_STRUCT(NiagaraConstantMap);
diff --git a/Engine/Source/Runtime/Niagara/Private/NiagaraEffect.cpp b/Engine/Source/Runtime/Niagara/Private/NiagaraEffect.cpp
new file mode 100644
--- /dev/null
+++ b/Engine/Source/Runtime/Niagara/Private/NiagaraEffect.cpp
@@ -0,0 +1,97 @@
+// Copyright 1998-2015 Epic Games, Inc. All Rights Reserved.
+
+
+#include "NiagaraPrivate.h"
+#include "NiagaraEffect.h"
+
+
+
+UNiagaraEffect::UNiagaraEffect(const FObjectInitializer& ObjectInitializer)
+: Super(ObjectInitializer)
+{
+}
+
+
+
+FNiagaraEmitterProperties *UNiagaraEffect::AddEmitterProperties()
+{
+	FNiagaraEmitterProperties *Props = new FNiagaraEmitterProperties();
+	EmitterProps.Add(Props);
+	return Props;
+}
+
+
+
+void UNiagaraEffect::CreateEffectRendererProps(TSharedPtr<FNiagaraSimulation> Sim)
+{
+	UClass *RendererProps = Sim->GetEffectRenderer()->GetPropertiesClass();
+	if (RendererProps)
+	{
+		Sim->GetProperties()->RendererProperties = NewObject<UNiagaraEffectRendererProperties>(this, RendererProps);
+	}
+	else
+	{
+		Sim->GetProperties()->RendererProperties = nullptr;
+	}
+}
+
+
+void UNiagaraEffect::PostLoad()
+{
+	Super::PostLoad();
+	EmitterProps.Empty();
+	for (FNiagaraEmitterProperties &Props : EmitterPropsSerialized)
+	{
+		FNiagaraEmitterProperties *NewProps = new FNiagaraEmitterProperties(Props);
+		EmitterProps.Add(NewProps);
+	}
+}
+
+
+void UNiagaraEffect::PreSave()
+{
+	EmitterPropsSerialized.Empty();
+	for (FNiagaraEmitterProperties *Props : EmitterProps)
+	{
+		EmitterPropsSerialized.Add(*Props);
+	}
+}
+
+
+
+
+TSharedPtr<FNiagaraSimulation> FNiagaraEffectInstance::AddEmitter(FNiagaraEmitterProperties *Properties)
+{
+	FNiagaraSimulation *SimPtr = new FNiagaraSimulation(Properties);
+	TSharedPtr<FNiagaraSimulation> Sim = MakeShareable(SimPtr);
+
+	Sim->SetRenderModuleType(RMT_Sprites, Component->GetWorld()->FeatureLevel);
+	Emitters.Add(Sim);
+	FNiagaraSceneProxy *SceneProxy = static_cast<FNiagaraSceneProxy*>(Component->SceneProxy);
+	SceneProxy->UpdateEffectRenderers(this);
+
+	return Sim;
+}
+
+
+void FNiagaraEffectInstance::RenderModuleupdate()
+{
+	FNiagaraSceneProxy *NiagaraProxy = static_cast<FNiagaraSceneProxy*>(Component->SceneProxy);
+	if (NiagaraProxy)
+	{
+		NiagaraProxy->UpdateEffectRenderers(this);
+	}
+}
+
+
+void FNiagaraEffectInstance::InitEmitters(UNiagaraEffect *InAsset)
+{
+	check(InAsset);
+	for (int i = 0; i < InAsset->GetNumEmitters(); i++)
+	{
+		FNiagaraEmitterProperties *Props = InAsset->GetEmitterProperties(i);
+		FNiagaraSimulation *Sim = new FNiagaraSimulation(Props);
+		Emitters.Add(MakeShareable(Sim));
+	}
+}
+
diff --git a/Engine/Source/Runtime/Niagara/Private/NiagaraModule.cpp b/Engine/Source/Runtime/Niagara/Private/NiagaraModule.cpp
new file mode 100644
--- /dev/null
+++ b/Engine/Source/Runtime/Niagara/Private/NiagaraModule.cpp
@@ -0,0 +1,6 @@
+// Copyright 1998-2015 Epic Games, Inc. All Rights Reserved.
+
+#include "NiagaraPrivate.h"
+#include "NiagaraModule.h"
+
+IMPLEMENT_MODULE(INiagaraModule, Niagara);
diff --git a/Engine/Source/Runtime/Niagara/Private/NiagaraPrivate.h b/Engine/Source/Runtime/Niagara/Private/NiagaraPrivate.h
new file mode 100644
--- /dev/null
+++ b/Engine/Source/Runtime/Niagara/Private/NiagaraPrivate.h
@@ -0,0 +1,6 @@
+// Copyright 1998-2015 Epic Games, Inc. All Rights Reserved.
+
+
+#pragma once
+#include "Engine.h"
+#include "VectorVM.h"
\ No newline at end of file
diff --git a/Engine/Source/Runtime/Niagara/Private/NiagaraScript.cpp b/Engine/Source/Runtime/Niagara/Private/NiagaraScript.cpp
new file mode 100644
--- /dev/null
+++ b/Engine/Source/Runtime/Niagara/Private/NiagaraScript.cpp
@@ -0,0 +1,19 @@
+// Copyright 1998-2015 Epic Games, Inc. All Rights Reserved.
+
+#include "NiagaraPrivate.h"
+#include "NiagaraScript.h"
+#include "NiagaraScriptSourceBase.h"
+
+
+
+UNiagaraScript::UNiagaraScript(const FObjectInitializer& ObjectInitializer)
+	: Super(ObjectInitializer)
+{
+
+}
+
+UNiagaraScriptSourceBase::UNiagaraScriptSourceBase(const FObjectInitializer& ObjectInitializer)
+	: Super(ObjectInitializer)
+{
+
+}
diff --git a/Engine/Source/Runtime/Niagara/Private/NiagaraSimulation.cpp b/Engine/Source/Runtime/Niagara/Private/NiagaraSimulation.cpp
new file mode 100644
--- /dev/null
+++ b/Engine/Source/Runtime/Niagara/Private/NiagaraSimulation.cpp
@@ -0,0 +1,337 @@
+// Copyright 1998-2015 Epic Games, Inc. All Rights Reserved.
+
+#include "NiagaraPrivate.h"
+#include "NiagaraSimulation.h"
+#include "NiagaraEffectRenderer.h"
+#include "VectorVM.h"
+
+
+FNiagaraSimulation::FNiagaraSimulation(FNiagaraEmitterProperties *InProps) 
+: Age(0.0f)
+, bIsEnabled(true)
+, SpawnRemainder(0.0f)
+, CachedBounds(ForceInit)
+, EffectRenderer(nullptr)
+{
+	Props = InProps;
+
+	Init();
+}
+
+FNiagaraSimulation::FNiagaraSimulation(FNiagaraEmitterProperties *InProps, ERHIFeatureLevel::Type InFeatureLevel)
+	: Age(0.0f)
+	, bIsEnabled(true)
+	, SpawnRemainder(0.0f)
+	, CachedBounds(ForceInit)
+	, EffectRenderer(nullptr)
+{
+	Props = InProps;
+
+	Init();
+	SetRenderModuleType(Props->RenderModuleType, InFeatureLevel);
+}
+
+void FNiagaraSimulation::Init()
+{
+	if (Props && Props->UpdateScript && Props->SpawnScript )
+	{
+		if (Props->UpdateScript->Attributes.Num() == 0 || Props->SpawnScript->Attributes.Num() == 0)
+		{
+			Data.Reset();
+			bIsEnabled = false;
+			UE_LOG(LogNiagara, Error, TEXT("This emitter cannot be enabled because it's spawn or update script doesn't have any attriubtes.."));
+		}
+		else
+		{
+			//Warn the user if there are any attributes used in the update script that are not initialized in the spawn script.
+			//TODO: We need some window in the effect editor and possibly the graph editor for warnings and errors.
+			for (FNiagaraVariableInfo& Attr : Props->UpdateScript->Attributes)
+			{
+				int32 FoundIdx;
+				if (!Props->SpawnScript->Attributes.Find(Attr, FoundIdx))
+				{
+					UE_LOG(LogNiagara, Warning, TEXT("Attribute %s is used in the Update script for %s but it is not initialised in the Spawn script!"), *Attr.Name.ToString(), *Props->Name);
+				}
+			}
+			Data.AddAttributes(Props->UpdateScript->Attributes);
+			Data.AddAttributes(Props->SpawnScript->Attributes);
+
+			Constants.Empty();
+			Constants.Merge(Props->ExternalConstants);
+		}
+	}
+	else
+	{
+		Data.Reset();//Clear the attributes to mark the sim as disabled independatly of the user set Enabled flag.
+		bIsEnabled = false;//Also force the user flag to give an indication to the user.
+		//TODO - Arbitrary named scripts. Would need some base functionality for Spawn/Udpate to be called that can be overriden in BPs for emitters with custom scripts.
+		UE_LOG(LogNiagara, Error, TEXT("This emitter cannot be enabled because it's doesn't have both an update and spawn script."));
+	}
+}
+
+void FNiagaraSimulation::SetProperties(FNiagaraEmitterProperties *InProps)	
+{ 
+	Props = InProps; 
+	Init();
+}
+
+void FNiagaraSimulation::SetSpawnScript(UNiagaraScript* Script)
+{
+	Props->SpawnScript = Script;
+	Init();
+}
+
+void FNiagaraSimulation::SetUpdateScript(UNiagaraScript* Script)
+{
+	Props->UpdateScript = Script;
+	Init();
+}
+
+void FNiagaraSimulation::SetEnabled(bool bEnabled)
+{
+	bIsEnabled = bEnabled;
+	Init();
+}
+
+
+float FNiagaraSimulation::GetTotalCPUTime()
+{
+	float Total = CPUTimeMS;
+	if (EffectRenderer)
+	{
+		Total += EffectRenderer->GetCPUTimeMS();
+	}
+
+	return Total;
+}
+
+int FNiagaraSimulation::GetTotalBytesUsed()
+{
+	return Data.GetBytesUsed();
+}
+
+
+void FNiagaraSimulation::Tick(float DeltaSeconds)
+{
+	SCOPE_CYCLE_COUNTER(STAT_NiagaraTick);
+
+	if (!bIsEnabled)
+		return;
+
+	SimpleTimer TickTime;
+
+	check(Data.GetNumAttributes() > 0);
+	check(Props->SpawnScript);
+	check(Props->UpdateScript);
+
+	// Cache the ComponentToWorld transform.
+//	CachedComponentToWorld = Component.GetComponentToWorld();
+
+	Data.SwapBuffers();
+	Data.SetNumParticles(Data.GetPrevNumParticles());
+
+	int32 OrigNumParticles = Data.GetNumParticles();
+	int32 NumToSpawn = 0;
+
+	// Figure out how many we will spawn.
+	NumToSpawn = CalcNumToSpawn(DeltaSeconds);
+	int32 MaxNewParticles = OrigNumParticles + NumToSpawn;
+	Data.Allocate(MaxNewParticles);
+
+	Age += DeltaSeconds;
+	Constants.SetOrAdd(TEXT("Emitter Age"), FVector4(Age, Age, Age, Age));
+	Constants.SetOrAdd(TEXT("Delta Time"), FVector4(DeltaSeconds, DeltaSeconds, DeltaSeconds, DeltaSeconds));
+
+	// Simulate particles forward by DeltaSeconds.
+	{
+		SCOPE_CYCLE_COUNTER(STAT_NiagaraSimulate);
+		RunVMScript(Props->UpdateScript, EUnusedAttributeBehaviour::Copy);
+	}
+	
+	//Init new particles with the spawn script.
+	{
+		SCOPE_CYCLE_COUNTER(STAT_NiagaraSpawn);
+		Data.SetNumParticles(MaxNewParticles);
+		//For now, zero any unused attributes here. But as this is really uninitialized data we should maybe make this a more serious error.
+		RunVMScript(Props->SpawnScript, EUnusedAttributeBehaviour::Zero, OrigNumParticles, NumToSpawn);
+	}
+
+	// Iterate over looking for dead particles and move from the end of the list to the dead location, compacting in the process
+	{
+		SCOPE_CYCLE_COUNTER(STAT_NiagaraKill);
+		int32 CurNumParticles = OrigNumParticles = Data.GetNumParticles();
+		int32 ParticleIndex = 0;
+		const FVector4* ParticleRelativeTimes = Data.GetAttributeData(FNiagaraVariableInfo(FName(TEXT("Age")), ENiagaraDataType::Vector));
+		if (ParticleRelativeTimes)
+		{
+			while (ParticleIndex < OrigNumParticles)
+			{
+				if (ParticleRelativeTimes[ParticleIndex].X > 1.0f)
+				{
+					// Particle is dead, move one from the end here. 
+					MoveParticleToIndex(--CurNumParticles, ParticleIndex);
+				}
+				ParticleIndex++;
+			}
+		}
+		Data.SetNumParticles(CurNumParticles);
+	}
+
+	CPUTimeMS = TickTime.GetElapsedMilliseconds();
+
+	DECLARE_DWORD_COUNTER_STAT(TEXT("NumParticles"), STAT_NiagaraNumParticles, STATGROUP_Niagara);
+	INC_DWORD_STAT_BY(STAT_NiagaraNumParticles, Data.GetNumParticles());
+
+}
+
+void FNiagaraSimulation::RunVMScript(UNiagaraScript* Script, EUnusedAttributeBehaviour UnusedAttribBehaviour)
+{
+	RunVMScript(Script, UnusedAttribBehaviour, 0, Data.GetNumParticles());
+}
+
+void FNiagaraSimulation::RunVMScript(UNiagaraScript* Script, EUnusedAttributeBehaviour UnusedAttribBehaviour, uint32 StartParticle)
+{
+	RunVMScript(Script, UnusedAttribBehaviour, StartParticle, Data.GetNumParticles() - StartParticle);
+}
+
+void FNiagaraSimulation::RunVMScript(UNiagaraScript* Script, EUnusedAttributeBehaviour UnusedAttribBehaviour, uint32 StartParticle, uint32 NumParticles)
+{
+	if (NumParticles == 0)
+		return;
+
+	check(Script);
+	check(Script->ByteCode.Num() > 0);
+	check(Script->Attributes.Num() > 0);
+	check(Data.GetNumAttributes() > 0);
+
+	uint32 CurrNumParticles = Data.GetNumParticles();
+
+	check(StartParticle + NumParticles <= CurrNumParticles);
+
+	const FVector4* InputBuffer = Data.GetPreviousBuffer();
+	const FVector4* OutputBuffer = Data.GetCurrentBuffer();
+
+	uint32 InputAttributeStride = Data.GetPrevParticleAllocation();
+	uint32 OutputAttributeStride = Data.GetParticleAllocation();
+
+	VectorRegister* InputRegisters[VectorVM::MaxInputRegisters] = { 0 };
+	VectorRegister* OutputRegisters[VectorVM::MaxOutputRegisters] = { 0 };
+	
+	//The current script run will be using NumScriptAttrs. We must pick out the Attributes form the simulation data in the order that they appear in the script.
+	//The remaining attributes being handled according to UnusedAttribBehaviour
+	const int32 NumSimulationAttrs = Data.GetNumAttributes();
+	const int32 NumScriptAttrs = Script->Attributes.Num();
+
+	check(NumScriptAttrs < VectorVM::MaxInputRegisters);
+	check(NumScriptAttrs < VectorVM::MaxOutputRegisters);
+
+	// Setup input and output registers.
+	for (TPair<FNiagaraVariableInfo, uint32> AttrIndexPair : Data.GetAttributes())
+	{
+		const FNiagaraVariableInfo& AttrInfo = AttrIndexPair.Key;
+
+		FVector4* InBuff;
+		FVector4* OutBuff;
+		Data.GetAttributeData(AttrInfo, InBuff, OutBuff, StartParticle);
+
+		int32 AttrIndex = Script->Attributes.Find(AttrInfo);
+		if (AttrIndex != INDEX_NONE)
+		{
+			//This attribute is used in this script so set it to the correct Input and Output buffer locations.
+			InputRegisters[AttrIndex] = (VectorRegister*)InBuff;
+			OutputRegisters[AttrIndex] = (VectorRegister*)OutBuff;
+		}	
+		else
+		{
+			//This attribute is not used in this script so handle it accordingly.
+			check(OutBuff != NULL);
+			if (UnusedAttribBehaviour == EUnusedAttributeBehaviour::Copy)
+			{
+				check(InBuff != NULL);
+				FMemory::Memcpy(OutBuff, InBuff, NumParticles * sizeof(FVector4));
+			}
+			else if (UnusedAttribBehaviour == EUnusedAttributeBehaviour::Zero)
+			{
+				FMemory::Memzero(OutBuff, NumParticles * sizeof(FVector4));
+			}
+			else if (UnusedAttribBehaviour == EUnusedAttributeBehaviour::MarkInvalid)
+			{
+				FMemory::Memset(OutBuff, NIAGARA_INVALID_MEMORY, NumParticles * sizeof(FVector4));
+			}
+		}
+	}
+
+	//Fill constant table with required emitter constants and internal script constants.
+	TArray<FVector4> ConstantTable;
+	TArray<FNiagaraDataObject *>DataObjTable;
+	Script->ConstantData.FillConstantTable(Constants, ConstantTable, DataObjTable);
+
+	VectorVM::Exec(
+		Script->ByteCode.GetData(),
+		InputRegisters,
+		NumScriptAttrs,
+		OutputRegisters,
+		NumScriptAttrs,
+		ConstantTable.GetData(),
+		DataObjTable.GetData(),
+		NumParticles
+		);
+}
+
+bool FNiagaraSimulation::CheckAttriubtesForRenderer()
+{
+	bool bOk = true;
+	if (EffectRenderer)
+	{
+		const TArray<FNiagaraVariableInfo>& RequiredAttrs = EffectRenderer->GetRequiredAttributes();
+
+		for (auto& Attr : RequiredAttrs)
+		{
+			if (!Data.HasAttriubte(Attr))
+			{
+				bOk = false;
+				UE_LOG(LogNiagara, Error, TEXT("Cannot render %s because it does not define attriubte %s."), *Props->Name, *Attr.Name.ToString());
+			}
+		}
+	}
+	return bOk;
+}
+
+/** Replace the current effect renderer with a new one of Type, transferring the existing material over
+ *	to the new renderer. Don't forget to call RenderModuleUpdate on the SceneProxy after calling this! 
+ */
+void FNiagaraSimulation::SetRenderModuleType(EEmitterRenderModuleType Type, ERHIFeatureLevel::Type FeatureLevel)
+{
+	if (Type != Props->RenderModuleType || EffectRenderer==nullptr)
+	{
+		UMaterial *Material = UMaterial::GetDefaultMaterial(MD_Surface);
+
+		if (EffectRenderer)
+		{
+			Material = EffectRenderer->GetMaterial();
+			delete EffectRenderer;
+		}
+		else
+		{
+			if (Props->Material)
+			{
+				Material = Props->Material;
+			}
+		}
+
+		Props->RenderModuleType = Type;
+		switch (Type)
+		{
+		case RMT_Sprites: EffectRenderer = new NiagaraEffectRendererSprites(FeatureLevel, Props->RendererProperties);
+			break;
+		case RMT_Ribbon: EffectRenderer = new NiagaraEffectRendererRibbon(FeatureLevel, Props->RendererProperties);
+			break;
+		default:	EffectRenderer = new NiagaraEffectRendererSprites(FeatureLevel, Props->RendererProperties);
+					Props->RenderModuleType = RMT_Sprites;
+					break;
+		}
+
+		EffectRenderer->SetMaterial(Material, FeatureLevel);
+		CheckAttriubtesForRenderer();
+	}
+}
diff --git a/Engine/Source/Runtime/Niagara/Public/NiagaraActor.h b/Engine/Source/Runtime/Niagara/Public/NiagaraActor.h
new file mode 100644
--- /dev/null
+++ b/Engine/Source/Runtime/Niagara/Public/NiagaraActor.h
@@ -0,0 +1,46 @@
+// Copyright 1998-2015 Epic Games, Inc. All Rights Reserved.
+
+#pragma once
+#include "NiagaraActor.generated.h"
+
+UCLASS(MinimalAPI)
+class ANiagaraActor : public AActor
+{
+	GENERATED_UCLASS_BODY()
+
+private_subobject:
+	/** Pointer to effect component */
+	DEPRECATED_FORGAME(4.6, "NiagaraComponent should not be accessed directly, please use GetNiagaraComponent() function instead. NiagaraComponent will soon be private and your code will not compile.")
+	UPROPERTY(VisibleAnywhere, Category=NiagaraActor)
+	class UNiagaraComponent* NiagaraComponent;
+
+#if WITH_EDITORONLY_DATA
+	// Reference to sprite visualization component
+	DEPRECATED_FORGAME(4.6, "SpriteComponent should not be accessed directly, please use GetSpriteComponent() function instead. SpriteComponent will soon be private and your code will not compile.")
+	UPROPERTY()
+	class UBillboardComponent* SpriteComponent;
+
+	// Reference to arrow visualization component
+	DEPRECATED_FORGAME(4.6, "ArrowComponent should not be accessed directly, please use GetArrowComponent() function instead. ArrowComponent will soon be private and your code will not compile.")
+	UPROPERTY()
+	class UArrowComponent* ArrowComponent;
+
+#endif
+
+public:
+	/** Returns NiagaraComponent subobject **/
+	NIAGARA_API class UNiagaraComponent* GetNiagaraComponent() const;
+#if WITH_EDITORONLY_DATA
+	/** Returns SpriteComponent subobject **/
+	class UBillboardComponent* GetSpriteComponent() const;
+	/** Returns ArrowComponent subobject **/
+	class UArrowComponent* GetArrowComponent() const;
+#endif
+
+#if WITH_EDITOR
+	// AActor interface
+	virtual bool GetReferencedContentObjects(TArray<UObject*>& Objects) const override;
+	// End of AActor interface
+#endif // WITH_EDITOR
+
+};
diff --git a/Engine/Source/Runtime/Niagara/Public/NiagaraCommon.h b/Engine/Source/Runtime/Niagara/Public/NiagaraCommon.h
new file mode 100644
--- /dev/null
+++ b/Engine/Source/Runtime/Niagara/Public/NiagaraCommon.h
@@ -0,0 +1,59 @@
+// Copyright 1998-2015 Epic Games, Inc. All Rights Reserved.
+
+#pragma once
+
+#include "NiagaraCommon.generated.h"
+
+DECLARE_LOG_CATEGORY_EXTERN(LogNiagara, Log, Verbose);
+
+UENUM()
+enum class ENiagaraDataType : uint8
+{
+	Scalar,
+	Vector,
+	Matrix,
+	Curve,
+};
+
+USTRUCT()
+struct FNiagaraVariableInfo
+{
+	GENERATED_USTRUCT_BODY()
+
+	FNiagaraVariableInfo()
+	: Name(NAME_None)
+	, Type(ENiagaraDataType::Vector)
+	{}
+
+	FNiagaraVariableInfo(FName InName, ENiagaraDataType InType)
+	: Name(InName)
+	, Type(InType)
+	{}
+
+	UPROPERTY(EditAnywhere, Category = "Variable")
+	FName Name;
+
+	UPROPERTY(EditAnywhere, Category = "Variable")
+	TEnumAsByte<ENiagaraDataType> Type;
+	
+	FORCEINLINE bool operator==(const FNiagaraVariableInfo& Other)const
+	{
+		return Name == Other.Name && Type == Other.Type;
+	}
+
+	FORCEINLINE bool operator!=(const FNiagaraVariableInfo& Other)const
+	{
+		return !(*this == Other);
+	}
+};
+
+FORCEINLINE FArchive& operator<<(FArchive& Ar, FNiagaraVariableInfo& VarInfo)
+{
+	Ar << VarInfo.Name << VarInfo.Type;
+	return Ar;
+}
+
+FORCEINLINE uint32 GetTypeHash(const FNiagaraVariableInfo& Var)
+{
+	return HashCombine(GetTypeHash(Var.Name), (uint32)(Var.Type.GetValue()));
+}
diff --git a/Engine/Source/Runtime/Niagara/Public/NiagaraComponent.h b/Engine/Source/Runtime/Niagara/Public/NiagaraComponent.h
new file mode 100644
--- /dev/null
+++ b/Engine/Source/Runtime/Niagara/Public/NiagaraComponent.h
@@ -0,0 +1,211 @@
+// Copyright 1998-2015 Epic Games, Inc. All Rights Reserved.
+
+#pragma once
+
+#include "NiagaraCommon.h"
+
+#include "NiagaraComponent.generated.h"
+
+/**
+* Contains the simulation data for a single FNiagaraSimulation
+* all buffers for particle attributes are stored here for convenient access
+*/
+class FNiagaraEmitterParticleData
+{
+public:
+	FNiagaraEmitterParticleData()
+		: CurrentBuffer(0)
+	{
+		Reset();
+	}
+
+	~FNiagaraEmitterParticleData() {}
+
+	void Allocate(uint32 NumExpectedParticles)
+	{
+		ParticleBuffers[CurrentBuffer].Reset(NumExpectedParticles * AttrMap.Num());
+		ParticleBuffers[CurrentBuffer].AddUninitialized(NumExpectedParticles * AttrMap.Num());
+
+		ParticleAllocation[CurrentBuffer] = NumExpectedParticles;
+	}
+
+	void Reset()
+	{
+		AttrMap.Empty();
+		NumParticles[0] = 0;
+		NumParticles[1] = 0;
+		ParticleAllocation[0] = 0;
+		ParticleAllocation[1] = 0;
+		CurrentBuffer = 0;
+	}
+
+	const FVector4* GetAttributeData(const FNiagaraVariableInfo& AttrID) const
+	{
+		const uint32* Offset = AttrMap.Find(AttrID);
+		return Offset ? ParticleBuffers[CurrentBuffer].GetData() + (*Offset * ParticleAllocation[CurrentBuffer]) : NULL;
+	}
+
+	void GetAttributeData(const FNiagaraVariableInfo& AttrID, FVector4*& PrevBuffer, FVector4*& CurrBuffer, int32 StartParticle) 
+	{
+		const uint32* Offset = AttrMap.Find(AttrID);
+		if (Offset)
+		{
+			PrevBuffer = ParticleBuffers[CurrentBuffer ^ 0x1].GetData() + (*Offset * ParticleAllocation[CurrentBuffer ^ 0x1]) + StartParticle;
+			CurrBuffer = ParticleBuffers[CurrentBuffer].GetData() + (*Offset * ParticleAllocation[CurrentBuffer]) + StartParticle;
+		}
+		else
+		{
+			PrevBuffer = NULL;
+			CurrBuffer = NULL;
+		}
+	}
+
+	bool HasAttriubte(const FNiagaraVariableInfo& AttrID)
+	{
+		return AttrMap.Find(AttrID) != NULL;
+	}
+
+	void SetAttributes(const TArray<FNiagaraVariableInfo>& Attrs)
+	{
+		Reset();
+		for (const FNiagaraVariableInfo& Attr : Attrs)
+		{
+			//TODO: Support attributes of any type. This may as well wait until the move of matrix ops etc over into UNiagaraScriptStructs.
+			uint32 Idx = AttrMap.Num();
+			AttrMap.Add(Attr, Idx);
+		}
+	}
+
+	void AddAttributes(const TArray<FNiagaraVariableInfo>& Attrs)
+	{
+		for (const FNiagaraVariableInfo& Attr : Attrs)
+		{
+			//TODO: Support attributes of any type. This may as well wait until the move of matrix ops etc over into UNiagaraScriptStructs.
+			uint32 Idx = AttrMap.Num();
+			if (!AttrMap.Find(Attr))
+			{
+				AttrMap.Add(Attr, Idx);
+			}
+		}
+	}
+
+	const TMap<FNiagaraVariableInfo, uint32>& GetAttributes()const { return AttrMap; }
+	int GetNumAttributes()				{ return AttrMap.Num(); }
+
+	uint32 GetNumParticles() const		{ return NumParticles[CurrentBuffer]; }
+	uint32 GetPrevNumParticles() const		{ return NumParticles[CurrentBuffer ^ 0x1]; }
+	uint32 GetParticleAllocation() const { return ParticleAllocation[CurrentBuffer]; }
+	uint32 GetPrevParticleAllocation() const { return ParticleAllocation[CurrentBuffer ^ 0x1]; }
+	void SetNumParticles(uint32 Num)	{ NumParticles[CurrentBuffer] = Num; }
+	void SwapBuffers()					{ CurrentBuffer ^= 0x1; }
+
+	FVector4* GetCurrentBuffer()		{ return ParticleBuffers[CurrentBuffer].GetData(); }
+	FVector4* GetPreviousBuffer()		{ return ParticleBuffers[CurrentBuffer^0x1].GetData(); }
+
+	int GetBytesUsed()	{ return (ParticleBuffers[0].Num() + ParticleBuffers[1].Num()) * 16 + AttrMap.Num() * 4; }
+private:
+	uint32 CurrentBuffer;
+	uint32 NumParticles[2];
+	uint32 ParticleAllocation[2];
+	TArray<FVector4> ParticleBuffers[2];
+	TMap<FNiagaraVariableInfo, uint32> AttrMap;
+};
+
+
+
+/**
+* UNiagaraComponent is the primitive component for a Niagara effect.
+* @see ANiagaraActor
+* @see UNiagaraEffect
+*/
+UCLASS(MinimalAPI)
+class UNiagaraComponent : public UPrimitiveComponent
+{
+	GENERATED_UCLASS_BODY()
+private:
+	UPROPERTY()
+	class UNiagaraEffect* Asset;
+
+	class FNiagaraEffectInstance* EffectInstance;
+
+	// Begin UActorComponent interface.
+protected:
+	virtual void OnRegister() override;
+	virtual void OnUnregister()  override;
+	virtual void SendRenderDynamicData_Concurrent() override;
+public:
+	virtual void TickComponent(float DeltaTime, enum ELevelTick TickType, FActorComponentTickFunction* ThisTickFunction) override;
+	virtual const UObject* AdditionalStatObject() const override;
+	// End UActorComponent interface.
+
+	// Begin UPrimitiveComponent Interface
+	virtual int32 GetNumMaterials() const override;
+	virtual FBoxSphereBounds CalcBounds(const FTransform& LocalToWorld) const override;
+	virtual FPrimitiveSceneProxy* CreateSceneProxy() override;
+	// End UPrimitiveComponent Interface
+
+	NIAGARA_API void SetAsset(UNiagaraEffect* InAsset);
+	UNiagaraEffect* GetAsset() const { return Asset; }
+
+	FNiagaraEffectInstance* GetEffectInstance()	const { return EffectInstance; }
+	void SetEffectInstance(FNiagaraEffectInstance* InInstance)	{ EffectInstance = InInstance; }
+
+	// Begin UObject interface.
+#if WITH_EDITOR
+	virtual void PostEditChangeProperty(FPropertyChangedEvent& PropertyChangedEvent) override;
+#endif // WITH_EDITOR
+	// End UObject interface.
+
+	NIAGARA_API static const TArray<FNiagaraVariableInfo>& GetSystemConstants();
+
+};
+
+
+
+
+
+
+/**
+* Scene proxy for drawing niagara particle simulations.
+*/
+class FNiagaraSceneProxy : public FPrimitiveSceneProxy
+{
+public:
+
+	FNiagaraSceneProxy(const UNiagaraComponent* InComponent);
+	~FNiagaraSceneProxy();
+
+	/** Called on render thread to assign new dynamic data */
+	void SetDynamicData_RenderThread(struct FNiagaraDynamicDataBase* NewDynamicData);
+	TArray<class NiagaraEffectRenderer*>& GetEffectRenderers() { return EffectRenderers; }
+	void AddEffectRenderer(NiagaraEffectRenderer* Renderer)	{ EffectRenderers.Add(Renderer); }
+	void UpdateEffectRenderers(FNiagaraEffectInstance* InEffect);
+
+private:
+	void ReleaseRenderThreadResources();
+
+	// FPrimitiveSceneProxy interface.
+	virtual void CreateRenderThreadResources() override;
+
+	virtual void OnActorPositionChanged() override;
+	virtual void OnTransformChanged() override;
+
+	virtual void GetDynamicMeshElements(const TArray<const FSceneView*>& Views, const FSceneViewFamily& ViewFamily, uint32 VisibilityMap, FMeshElementCollector& Collector) const override;
+
+	virtual FPrimitiveViewRelevance GetViewRelevance(const FSceneView* View)  override;
+
+	/*
+	virtual bool CanBeOccluded() const override
+	{
+	return !MaterialRelevance.bDisableDepthTest;
+	}
+	*/
+	virtual uint32 GetMemoryFootprint() const override;
+
+	uint32 GetAllocatedSize() const;
+
+
+private:
+	//class NiagaraEffectRenderer* EffectRenderer;
+	TArray<class NiagaraEffectRenderer*>EffectRenderers;
+};
diff --git a/Engine/Source/Runtime/Niagara/Public/NiagaraModule.h b/Engine/Source/Runtime/Niagara/Public/NiagaraModule.h
new file mode 100644
--- /dev/null
+++ b/Engine/Source/Runtime/Niagara/Public/NiagaraModule.h
@@ -0,0 +1,13 @@
+// Copyright 1998-2015 Epic Games, Inc. All Rights Reserved.
+
+#pragma once
+
+#include "ModuleInterface.h"
+
+/**
+* Foliage Edit mode module interface
+*/
+class INiagaraModule : public IModuleInterface
+{
+public:
+};
diff --git a/Engine/Source/Runtime/VectorVM/Private/VectorVM.cpp b/Engine/Source/Runtime/VectorVM/Private/VectorVM.cpp
--- a/Engine/Source/Runtime/VectorVM/Private/VectorVM.cpp
+++ b/Engine/Source/Runtime/VectorVM/Private/VectorVM.cpp
@@ -1,6 +1,8 @@
 // Copyright 1998-2015 Epic Games, Inc. All Rights Reserved.
 
+#include "CurveVector.h"
 #include "VectorVMPrivate.h"
+#include "VectorVMDataObject.h"
 #include "ModuleManager.h"
 
 IMPLEMENT_MODULE(FDefaultModuleImpl, VectorVM);
@@ -772,6 +774,42 @@ struct FVectorKernelCompose : public TQuatenaryVectorKernel<FVectorKernelCompose
 	}
 };
 
+// Ken Perlin's smootherstep function (zero first and second order derivatives at 0 and 1)
+// calculated separately for each channel of Src2
+struct FVectorKernelEaseIn : public TTrinaryVectorKernel<FVectorKernelEaseIn>
+{
+	static void VM_FORCEINLINE DoKernel(VectorRegister* RESTRICT Dst, const VectorRegister& Src0, const VectorRegister& Src1, const VectorRegister& Src2)
+	{
+		VectorRegister X = VectorMin( VectorDivide(VectorSubtract(Src2, Src0), VectorSubtract(Src1, Src0)), MakeVectorRegister(1.0f, 1.0f, 1.0f, 1.0f) );
+		X = VectorMax(X, MakeVectorRegister(0.0f, 0.0f, 0.0f, 0.0f));
+
+		VectorRegister X3 = VectorMultiply( VectorMultiply(X, X), X);
+		VectorRegister N6 = MakeVectorRegister(6.0f, 6.0f, 6.0f, 6.0f);
+		VectorRegister N15 = MakeVectorRegister(15.0f, 15.0f, 15.0f, 15.0f);
+		VectorRegister T = VectorSubtract( VectorMultiply(X, N6), N15 );
+		T = VectorAdd(VectorMultiply(X, T), MakeVectorRegister(10.0f, 10.0f, 10.0f, 10.0f));
+
+		*Dst = VectorMultiply(X3, T);
+	}
+};
+
+// smoothly runs 0->1->0
+struct FVectorKernelEaseInOut : public TUnaryVectorKernel<FVectorKernelEaseInOut>
+{
+	static void VM_FORCEINLINE DoKernel(VectorRegister* RESTRICT Dst, const VectorRegister& Src0)
+	{
+		VectorRegister T = VectorMultiply(Src0, MakeVectorRegister(2.0f, 2.0f, 2.0f, 2.0f));
+		T = VectorSubtract(T, MakeVectorRegister(1.0f, 1.0f, 1.0f, 1.0f));
+		VectorRegister X2 = VectorMultiply(T, T);
+		
+		VectorRegister R = VectorMultiply(X2, MakeVectorRegister(0.9604f, 0.9604f, 0.9604f, 0.9604f));
+		R = VectorSubtract(R, MakeVectorRegister(1.96f, 1.96f, 1.96f, 1.96f));
+		R = VectorMultiply(R, X2);
+		*Dst = VectorAdd(R, MakeVectorRegister(1.0f, 1.0f, 1.0f, 1.0f));
+	}
+};
+
+
 struct FVectorKernelOutput : public TUnaryVectorKernel<FVectorKernelOutput>
 {
 	static void VM_FORCEINLINE DoKernel(VectorRegister* Dst, VectorRegister Src0)
@@ -901,8 +939,11 @@ void VectorVM::Exec(
 			case EOp::lessthan: FVectorKernelLessThan::Exec(Context); break;
 			case EOp::sample: FVectorKernelSample::Exec(Context); break;
 			case EOp::bufferwrite: FVectorKernelBufferWrite::Exec(Context); break;
+			case EOp::eventbroadcast: break;
+			case EOp::easein: FVectorKernelEaseIn::Exec(Context); break;
+			case EOp::easeinout: FVectorKernelEaseInOut::Exec(Context); break;
 			case EOp::output: FVectorKernelOutput::Exec(Context); break;
-
+				
 			// Execution always terminates with a "done" opcode.
 			case EOp::done:
 				break;
diff --git a/Engine/Source/Runtime/VectorVM/Public/VectorVM.h b/Engine/Source/Runtime/VectorVM/Public/VectorVM.h
--- a/Engine/Source/Runtime/VectorVM/Public/VectorVM.h
+++ b/Engine/Source/Runtime/VectorVM/Public/VectorVM.h
@@ -1,90 +1,9 @@
 // Copyright 1998-2015 Epic Games, Inc. All Rights Reserved.
 
 #pragma once
-
 #include "Core.h"
 
-
-/* Vector VM data object; encapsulates buffers, curves and other data in its derivatives
-*  for access by VectorVM kernels;
-*/
-class FNiagaraDataObject
-{
-public:
-	virtual FVector4 Sample(const FVector4& InCoords) const = 0;
-	virtual FVector4 Write(const FVector4& InCoords, const FVector4 & Value) = 0;
-};
-
-/* Curve object; encapsulates a curve for the VectorVM
-*/
-class FNiagaraCurveDataObject : public FNiagaraDataObject
-{
-private:
-	class UCurveVector *CurveObj;
-public:
-	FNiagaraCurveDataObject(class UCurveVector *InCurve) : CurveObj(InCurve)
-	{
-	}
-
-	virtual FVector4 Sample(const FVector4& InCoords) const override
-	{
-		// commented out for the time being; until Niagara is in its own module, we can't depend on curve objects here
-		// because they're part of engine
-		//FVector Vec = CurveObj->GetVectorValue(InTime);	
-		//return FVector4(Vec, 0.0f);
-		return FVector4(1.0f, 0.0f, 0.0f, 1.0f);
-	}
-
-	virtual FVector4 Write(const FVector4& InCoords, const FVector4& Value) override
-	{
-		return FVector4(1.0f, 0.0f, 0.0f, 1.0f);
-	}
-
-	UCurveVector *GetCurveObject()	{ return CurveObj; }
-	void SetCurveObject(UCurveVector *InCurve)	{ CurveObj = InCurve; }
-};
-
-
-/* Curve object; encapsulates sparse volumetric data for the VectorVM
-*/
-class FNiagaraSparseVolumeDataObject : public FNiagaraDataObject
-{
-private:
-	TArray<FVector4> Data;
-public:
-	FNiagaraSparseVolumeDataObject()
-	{
-		Data.AddZeroed(32768);
-	}
-
-	uint32 MakeHash(const FVector4& InCoords) const
-	{
-		FVector4 Coords = InCoords;
-		const uint32 P1 = 73856093; // some large primes
-		const uint32 P2 = 19349663;
-		const uint32 P3 = 83492791;
-		Coords += FVector4(500.0f, 500.0f, 500.0f, 500.0f);
-		Coords *= 0.02f;
-		Coords = FVector4(FMath::Max(Coords.X, 0.0f), FMath::Max(Coords.Y, 0.0f), FMath::Max(Coords.Z, 0.0f), FMath::Max(Coords.W, 0.0f));
-		int N = (P1*static_cast<uint32>(Coords.X)) ^ (P2*static_cast<uint32>(Coords.Y)) ^ (P3*static_cast<uint32>(Coords.Z));
-		N %= 32768;
-		return FMath::Clamp(N, 0, 32768);
-	}
-
-	virtual FVector4 Sample(const FVector4& Coords) const override
-	{
-		int32 Index = MakeHash(Coords);
-		Index = FMath::Clamp(Index, 0, Data.Num() - 1);
-		return Data[Index];
-	}
-
-	virtual FVector4 Write(const FVector4& InCoords, const FVector4& Value) override
-	{
-		Data[MakeHash(InCoords)] = Value;
-		return Value;
-	}
-};
-
+class FNiagaraDataObject;
 
 namespace VectorVM
 {
@@ -166,6 +85,8 @@ namespace VectorVM
 			sample,
 			bufferwrite,
 			eventbroadcast,
+			easein,
+			easeinout,
 			NumOpcodes
 	};
 
diff --git a/Engine/Source/Runtime/VectorVM/VectorVM.Build.cs b/Engine/Source/Runtime/VectorVM/VectorVM.Build.cs
--- a/Engine/Source/Runtime/VectorVM/VectorVM.Build.cs
+++ b/Engine/Source/Runtime/VectorVM/VectorVM.Build.cs
@@ -6,6 +6,19 @@ public class VectorVM : ModuleRules
 {
 	public VectorVM(TargetInfo Target)
 	{
-		PrivateDependencyModuleNames.Add("Core");
-	}
+		PrivateDependencyModuleNames.AddRange(
+            new string[] {
+                "Core",
+                "Engine"
+            }
+        );
+
+        PrivateIncludePaths.AddRange(
+            new string[] {
+                "Runtime/Engine/Classes/Curves"
+            }
+        );
+
+
+    }
 }
