Niagara BP spawning, constant setting (experimental!)

diff --git a/Engine/Source/Runtime/Niagara/Classes/NiagaraEffect.h b/Engine/Source/Runtime/Niagara/Classes/NiagaraEffect.h
--- a/Engine/Source/Runtime/Niagara/Classes/NiagaraEffect.h
+++ b/Engine/Source/Runtime/Niagara/Classes/NiagaraEffect.h
@@ -8,13 +8,13 @@
 #include "NiagaraEffect.generated.h"
 
 
-UCLASS(MinimalAPI)
-class UNiagaraEffect : public UObject
+UCLASS()
+class NIAGARA_API UNiagaraEffect : public UObject
 {
 	GENERATED_UCLASS_BODY()
 
 public:
-	NIAGARA_API FNiagaraEmitterProperties* AddEmitterProperties();
+	FNiagaraEmitterProperties* AddEmitterProperties();
 
 	void Init()
 	{
@@ -32,7 +32,7 @@ public:
 		return EmitterProps.Num();
 	}
 
-	NIAGARA_API void CreateEffectRendererProps(TSharedPtr<FNiagaraSimulation> Sim);
+	void CreateEffectRendererProps(TSharedPtr<FNiagaraSimulation> Sim);
 
 	virtual void PostLoad() override; 
 	virtual void PreSave() override;
diff --git a/Engine/Source/Runtime/Niagara/Private/NiagaraFunctionLibrary.cpp b/Engine/Source/Runtime/Niagara/Private/NiagaraFunctionLibrary.cpp
new file mode 100644
--- /dev/null
+++ b/Engine/Source/Runtime/Niagara/Private/NiagaraFunctionLibrary.cpp
@@ -0,0 +1,95 @@
+// Copyright 1998-2014 Epic Games, Inc. All Rights Reserved.
+
+#include "NiagaraPrivate.h"
+#include "NiagaraFunctionLibrary.h"
+
+
+UNiagaraFunctionLibrary::UNiagaraFunctionLibrary(const FObjectInitializer& ObjectInitializer)
+: Super(ObjectInitializer)
+{
+}
+
+
+/**
+* Spawns a Niagara effect at the specified world location/rotation
+* @return			The spawned UNiagaraComponent
+*/
+UNiagaraComponent* UNiagaraFunctionLibrary::SpawnEffectAtLocation(UObject* WorldContextObject, UNiagaraEffect* EffectTemplate, FVector SpawnLocation, FRotator SpawnRotation, bool bAutoDestroy)
+{
+	UNiagaraComponent* PSC = NULL;
+	if (EffectTemplate)
+	{
+		UWorld* World = GEngine->GetWorldFromContextObject(WorldContextObject);
+		if (World != nullptr)
+		{
+			AActor* Actor = World->GetWorldSettings();
+			UNiagaraComponent* PSC = NewObject<UNiagaraComponent>((Actor ? Actor : (UObject*)World));
+			PSC->SetAsset(EffectTemplate);
+			PSC->RegisterComponentWithWorld(World);
+
+			PSC->SetAbsolute(true, true, true);
+			PSC->SetWorldLocationAndRotation(SpawnLocation, SpawnRotation);
+			PSC->SetRelativeScale3D(FVector(1.f));
+			//PSC->ActivateSystem(true);
+		}
+	}
+	return PSC;
+}
+
+
+
+
+
+/**
+* Spawns a Niagara effect attached to a component
+* @return			The spawned UNiagaraComponent
+*/
+UNiagaraComponent* UNiagaraFunctionLibrary::SpawnEffectAttached(UNiagaraEffect* EffectTemplate, USceneComponent* AttachToComponent, FName AttachPointName, FVector Location, FRotator Rotation, EAttachLocation::Type LocationType, bool bAutoDestroy)
+{
+	UNiagaraComponent* PSC = nullptr;
+	if (EffectTemplate)
+	{
+		if (AttachToComponent == NULL)
+		{
+			UE_LOG(LogScript, Warning, TEXT("UNiagaraFunctionLibrary::SpawnEffectAttached: NULL AttachComponent specified!"));
+		}
+		else
+		{
+			AActor* Actor = AttachToComponent->GetOwner();
+			PSC = NewObject<UNiagaraComponent>((Actor ? Actor : (UObject*)AttachToComponent->GetWorld()));
+			PSC->SetAsset(EffectTemplate);
+			PSC->RegisterComponentWithWorld(AttachToComponent->GetWorld());
+
+			PSC->AttachTo(AttachToComponent, AttachPointName);
+			if (LocationType == EAttachLocation::KeepWorldPosition)
+			{
+				PSC->SetWorldLocationAndRotation(Location, Rotation);
+			}
+			else
+			{
+				PSC->SetRelativeLocationAndRotation(Location, Rotation);
+			}
+			PSC->SetRelativeScale3D(FVector(1.f));
+		}
+	}
+	return PSC;
+}
+
+
+/**
+* Set a constant in an emitter of a Niagara effect
+*/
+void UNiagaraFunctionLibrary::SetUpdateScriptConstant(UNiagaraComponent* Component, FName EmitterName, FName ConstantName, FVector Value)
+{
+	TArray<TSharedPtr<FNiagaraSimulation>> &Emitters = Component->GetEffectInstance()->GetEmitters();
+
+	for (TSharedPtr<FNiagaraSimulation> &Emitter : Emitters)
+	{
+		FName CurName = *Emitter->GetProperties()->Name;
+		if (CurName == EmitterName)
+		{
+			Emitter->GetProperties()->ExternalConstants.SetOrAdd(ConstantName, Value);
+			break;
+		}
+	}
+}
\ No newline at end of file
diff --git a/Engine/Source/Runtime/Niagara/Public/NiagaraComponent.h b/Engine/Source/Runtime/Niagara/Public/NiagaraComponent.h
--- a/Engine/Source/Runtime/Niagara/Public/NiagaraComponent.h
+++ b/Engine/Source/Runtime/Niagara/Public/NiagaraComponent.h
@@ -118,8 +118,8 @@ private:
 * @see ANiagaraActor
 * @see UNiagaraEffect
 */
-UCLASS(MinimalAPI)
-class UNiagaraComponent : public UPrimitiveComponent
+UCLASS()
+class NIAGARA_API UNiagaraComponent : public UPrimitiveComponent
 {
 	GENERATED_UCLASS_BODY()
 private:
@@ -144,7 +144,7 @@ public:
 	virtual FPrimitiveSceneProxy* CreateSceneProxy() override;
 	// End UPrimitiveComponent Interface
 
-	NIAGARA_API void SetAsset(UNiagaraEffect* InAsset);
+	void SetAsset(UNiagaraEffect* InAsset);
 	UNiagaraEffect* GetAsset() const { return Asset; }
 
 	FNiagaraEffectInstance* GetEffectInstance()	const { return EffectInstance; }
@@ -156,8 +156,7 @@ public:
 #endif // WITH_EDITOR
 	// End UObject interface.
 
-	NIAGARA_API static const TArray<FNiagaraVariableInfo>& GetSystemConstants();
-
+	static const TArray<FNiagaraVariableInfo>& GetSystemConstants();
 };
 
 
diff --git a/Engine/Source/Runtime/Niagara/Public/NiagaraFunctionLibrary.h b/Engine/Source/Runtime/Niagara/Public/NiagaraFunctionLibrary.h
new file mode 100644
--- /dev/null
+++ b/Engine/Source/Runtime/Niagara/Public/NiagaraFunctionLibrary.h
@@ -0,0 +1,30 @@
+// Copyright 1998-2013 Epic Games, Inc. All Rights Reserved.
+
+#pragma once
+
+#include "Engine.h" // needed for access to UBlueprintFunctionLibrary
+#include "NiagaraFunctionLibrary.generated.h"
+
+
+/**
+* A C++ and Blueprint accessible library of utility functions for accessing Leap Motion device.
+* All positions & orientations are returned in Unreal reference frame & units, assuming the Leap device is located at the origin.
+*/
+UCLASS()
+class NIAGARA_API UNiagaraFunctionLibrary : public UBlueprintFunctionLibrary
+{
+	GENERATED_UCLASS_BODY()
+
+	/**
+	* Spawns a Niagara effect at the specified world location/rotation
+	* @return			The spawned UNiagaraComponent
+	*/
+	UFUNCTION(BlueprintCallable, Category = Niagara, meta = (Keywords = "niagara effect", WorldContext = "WorldContextObject", UnsafeDuringActorConstruction = "true"))
+	static UNiagaraComponent* SpawnEffectAtLocation(UObject* WorldContextObject, class UNiagaraEffect* EffectTemplate, FVector Location, FRotator Rotation = FRotator::ZeroRotator, bool bAutoDestroy = true);
+
+	UFUNCTION(BlueprintCallable, Category = Niagara, meta = (Keywords = "niagara effect", UnsafeDuringActorConstruction = "true"))
+	static UNiagaraComponent* SpawnEffectAttached(UNiagaraEffect* EffectTemplate, USceneComponent* AttachToComponent, FName AttachPointName, FVector Location, FRotator Rotation, EAttachLocation::Type LocationType, bool bAutoDestroy);
+
+	UFUNCTION(BlueprintCallable, Category = Niagara, meta = (Keywords = "niagara effect", UnsafeDuringActorConstruction = "true"))
+	static void SetUpdateScriptConstant(UNiagaraComponent* Component, FName EmitterName, FName ConstantName, FVector Value);
+};
\ No newline at end of file
