CIS fix

diff --git a/Engine/Source/Runtime/Niagara/Private/NiagaraFunctionLibrary.cpp b/Engine/Source/Runtime/Niagara/Private/NiagaraFunctionLibrary.cpp
--- a/Engine/Source/Runtime/Niagara/Private/NiagaraFunctionLibrary.cpp
+++ b/Engine/Source/Runtime/Niagara/Private/NiagaraFunctionLibrary.cpp
@@ -2,6 +2,9 @@
 
 #include "NiagaraPrivate.h"
 #include "NiagaraFunctionLibrary.h"
+#include "NiagaraComponent.h"
+#include "NiagaraEffect.h"
+#include "NiagaraSimulation.h"
 
 
 UNiagaraFunctionLibrary::UNiagaraFunctionLibrary(const FObjectInitializer& ObjectInitializer)
