Fix UNiagaraFunctionLibrary::SpawnEffectAtLocation not returning the created component (due to shadow variable)

diff --git a/Engine/Source/Runtime/Niagara/Private/NiagaraFunctionLibrary.cpp b/Engine/Source/Runtime/Niagara/Private/NiagaraFunctionLibrary.cpp
--- a/Engine/Source/Runtime/Niagara/Private/NiagaraFunctionLibrary.cpp
+++ b/Engine/Source/Runtime/Niagara/Private/NiagaraFunctionLibrary.cpp
@@ -26,7 +26,7 @@ UNiagaraComponent* UNiagaraFunctionLibrary::SpawnEffectAtLocation(UObject* World
 		if (World != nullptr)
 		{
 			AActor* Actor = World->GetWorldSettings();
-			UNiagaraComponent* PSC = NewObject<UNiagaraComponent>((Actor ? Actor : (UObject*)World));
+			PSC = NewObject<UNiagaraComponent>((Actor ? Actor : (UObject*)World));
 			PSC->SetAsset(EffectTemplate);
 			PSC->RegisterComponentWithWorld(World);
 
