Niagara:
- Fixing two potential crashes on compiling scripts and changing materials of effects that exist in the world
- Disabling rendering (not just ticking) of disabled emitters

diff --git a/Engine/Source/Runtime/Niagara/Private/NiagaraComponent.cpp b/Engine/Source/Runtime/Niagara/Private/NiagaraComponent.cpp
--- a/Engine/Source/Runtime/Niagara/Private/NiagaraComponent.cpp
+++ b/Engine/Source/Runtime/Niagara/Private/NiagaraComponent.cpp
@@ -203,8 +203,10 @@ void UNiagaraComponent::OnRegister()
 					{
 						Emitter->SetRenderModuleType(Emitter->GetProperties()->RenderModuleType, InComponent->GetWorld()->FeatureLevel);
 					}
+					InEffect->RenderModuleupdate();
 				}
 			);
+			FlushRenderingCommands();
 		}
 		VectorVM::Init();
 	}
@@ -227,15 +229,27 @@ void UNiagaraComponent::SendRenderDynamicData_Concurrent()
 			NiagaraEffectRenderer* Renderer = Emitter->GetEffectRenderer();
 			if (Renderer)
 			{
-				FNiagaraDynamicDataBase* DynamicData = Renderer->GenerateVertexData(Emitter->GetData());
-
-				ENQUEUE_UNIQUE_RENDER_COMMAND_TWOPARAMETER(
-					FSendNiagaraDynamicData,
-					NiagaraEffectRenderer*, EffectRenderer, Emitter->GetEffectRenderer(),
-					FNiagaraDynamicDataBase*, DynamicData, DynamicData,
-					{
-					EffectRenderer->SetDynamicData_RenderThread(DynamicData);
-				});
+				if (Emitter->IsEnabled())
+				{
+					FNiagaraDynamicDataBase* DynamicData = Renderer->GenerateVertexData(Emitter->GetData());
+
+					ENQUEUE_UNIQUE_RENDER_COMMAND_TWOPARAMETER(
+						FSendNiagaraDynamicData,
+						NiagaraEffectRenderer*, EffectRenderer, Emitter->GetEffectRenderer(),
+						FNiagaraDynamicDataBase*, DynamicData, DynamicData,
+						{
+						EffectRenderer->SetDynamicData_RenderThread(DynamicData);
+					});
+				}
+				else
+				{
+					ENQUEUE_UNIQUE_RENDER_COMMAND_ONEPARAMETER(
+						FSendNiagaraDynamicData,
+						NiagaraEffectRenderer*, EffectRenderer, Emitter->GetEffectRenderer(),
+						{
+							EffectRenderer->SetDynamicData_RenderThread(nullptr);
+						});
+				}
 			}
 		}
 	}
