Fixed Linux build (members must be initialized in the order they're declared)

diff --git a/Engine/Source/Runtime/Niagara/Classes/NiagaraSimulation.h b/Engine/Source/Runtime/Niagara/Classes/NiagaraSimulation.h
--- a/Engine/Source/Runtime/Niagara/Classes/NiagaraSimulation.h
+++ b/Engine/Source/Runtime/Niagara/Classes/NiagaraSimulation.h
@@ -39,8 +39,15 @@ struct FNiagaraEmitterProperties
 	GENERATED_USTRUCT_BODY()
 public:
 	FNiagaraEmitterProperties() : 
-		bIsEnabled(true), SpawnRate(50), UpdateScript(nullptr), SpawnScript(nullptr), Material(nullptr), 
-		RenderModuleType(RMT_Sprites), RendererProperties(nullptr), StartTime(0.0f), EndTime(0.0f)
+		bIsEnabled(true),
+		SpawnRate(50),
+		UpdateScript(nullptr),
+		SpawnScript(nullptr),
+		Material(nullptr), 
+		RenderModuleType(RMT_Sprites),
+		StartTime(0.0f),
+		EndTime(0.0f),
+		RendererProperties(nullptr)
 	{
 		Name = FString(TEXT("New Emitter"));
 	}
