Fixing Clang compile error.

diff --git a/Engine/Source/Runtime/Niagara/Classes/NiagaraEffect.h b/Engine/Source/Runtime/Niagara/Classes/NiagaraEffect.h
--- a/Engine/Source/Runtime/Niagara/Classes/NiagaraEffect.h
+++ b/Engine/Source/Runtime/Niagara/Classes/NiagaraEffect.h
@@ -144,8 +144,9 @@ public:
 
 			// manage emitter lifetime
 			//
-			if (Props->StartTime == 0.0f && Props->EndTime == 0.0f
-				|| (LoopedStartTime<Age && LoopedEndTime>Age))
+			if (	(Props->StartTime == 0.0f && Props->EndTime == 0.0f)
+					|| (LoopedStartTime<Age && LoopedEndTime>Age)
+				)
 			{
 				it->SetTickState(NTS_Running);
 			}
