Added a new, alternate UI to the Niagara effect editor. Accesible as a new tab under the windows menu.
This arranges all emitters in an SDetailView with some customization.
Primary benefit of this is that hooking things into the UI will require almost no work and elements can still be customized for a good UX.
Previous UI is still available and should work exactly as before.

Quite a bit of refactoring was needed to acheive this but much of it was needed for other things anyway and/or has some nice additonal benefits.
E.g.
- All constant types are supported easily.
- You can now create a default data object in the graph which will propagate through to the effect editor.
- Adding new data object types will be very easy.

Also added a delete and duplicate button for emitters in both UIs.

diff --git a/Engine/Source/Editor/NiagaraEditor/NiagaraEditor.Build.cs b/Engine/Source/Editor/NiagaraEditor/NiagaraEditor.Build.cs
--- a/Engine/Source/Editor/NiagaraEditor/NiagaraEditor.Build.cs
+++ b/Engine/Source/Editor/NiagaraEditor/NiagaraEditor.Build.cs
@@ -23,6 +23,7 @@ public class NiagaraEditor : ModuleRules
                 "Niagara",
                 "MovieSceneCore",
 				"Sequencer",
+				"PropertyEditor",
 			}
 		);
 
@@ -30,6 +31,7 @@ public class NiagaraEditor : ModuleRules
             new string[] {
 				"Engine",
                 "Niagara",
+				"PropertyEditor",
             }
         );
 
@@ -44,6 +46,7 @@ public class NiagaraEditor : ModuleRules
 		DynamicallyLoadedModuleNames.AddRange(
             new string[] {
                 "WorkspaceMenuStructure",
+				"PropertyEditor",
                 }
             );
 	}
diff --git a/Engine/Source/Editor/NiagaraEditor/Private/NiagaraCompiler.cpp b/Engine/Source/Editor/NiagaraEditor/Private/NiagaraCompiler.cpp
--- a/Engine/Source/Editor/NiagaraEditor/Private/NiagaraCompiler.cpp
+++ b/Engine/Source/Editor/NiagaraEditor/Private/NiagaraCompiler.cpp
@@ -331,6 +331,12 @@ TNiagaraExprPtr FNiagaraCompiler::GetExternalConstant(const FNiagaraVariableInfo
 	SetOrAddConstant(false, Constant, Default);
 	return Expression_GetExternalConstant(Constant);
 }
+TNiagaraExprPtr FNiagaraCompiler::GetExternalConstant(const FNiagaraVariableInfo& Constant, const UNiagaraDataObject* Default)
+{
+	UNiagaraDataObject* DefaultDupe = Default ? CastChecked<UNiagaraDataObject>(StaticDuplicateObject(Default, Script, NULL, ~RF_Transient)) : nullptr;
+	SetOrAddConstant(false, Constant, DefaultDupe);
+	return Expression_GetExternalConstant(Constant);
+}
 TNiagaraExprPtr FNiagaraCompiler::GetInternalConstant(const FNiagaraVariableInfo& Constant, float Default)
 {
 	SetOrAddConstant(true, Constant, Default);
@@ -346,11 +352,11 @@ TNiagaraExprPtr FNiagaraCompiler::GetInternalConstant(const FNiagaraVariableInfo
 	SetOrAddConstant(true, Constant, Default);
 	return Expression_GetInternalConstant(Constant);
 }
-
-TNiagaraExprPtr FNiagaraCompiler::GetExternalCurveConstant(const FNiagaraVariableInfo& Constant)
+TNiagaraExprPtr FNiagaraCompiler::GetInternalConstant(const FNiagaraVariableInfo& Constant, const UNiagaraDataObject* Default)
 {
-	SetOrAddConstant<UNiagaraDataObject*>(false, Constant, nullptr);
-	return Expression_GetExternalConstant(Constant);
+	UNiagaraDataObject* DefaultDupe = Default ? CastChecked<UNiagaraDataObject>(StaticDuplicateObject(Default, Script, NULL, ~RF_Transient)) : nullptr;
+	SetOrAddConstant(true, Constant, DefaultDupe);
+	return Expression_GetInternalConstant(Constant);
 }
 
 // TODO: Refactor this (and some other stuff) into a preprocessing step for use by any compiler?
@@ -413,9 +419,12 @@ bool FNiagaraCompiler::MergeInFunctionNodes()
 				FName OutputName(*FuncCallOutputPin->PinName);
 				for (UEdGraphPin* LinkTo : FuncCallOutputPin->LinkedTo)
 				{
-					check(LinkTo->Direction == EGPD_Input);
-					FReconnectionInfo& OutputConnection = OutputConnections.FindOrAdd(OutputName);
-					OutputConnection.To.Add(LinkTo);
+					if (LinkTo)
+					{
+						check(LinkTo->Direction == EGPD_Input);
+						FReconnectionInfo& OutputConnection = OutputConnections.FindOrAdd(OutputName);
+						OutputConnection.To.Add(LinkTo);
+					}
 				}
 			}
 
diff --git a/Engine/Source/Editor/NiagaraEditor/Private/NiagaraCompiler.h b/Engine/Source/Editor/NiagaraEditor/Private/NiagaraCompiler.h
--- a/Engine/Source/Editor/NiagaraEditor/Private/NiagaraCompiler.h
+++ b/Engine/Source/Editor/NiagaraEditor/Private/NiagaraCompiler.h
@@ -62,10 +62,11 @@ public:
 	virtual TNiagaraExprPtr GetExternalConstant(const FNiagaraVariableInfo& Constant, float Default)override;
 	virtual TNiagaraExprPtr GetExternalConstant(const FNiagaraVariableInfo& Constant, const FVector4& Default)override;
 	virtual TNiagaraExprPtr GetExternalConstant(const FNiagaraVariableInfo& Constant, const FMatrix& Default)override;
+	virtual TNiagaraExprPtr GetExternalConstant(const FNiagaraVariableInfo& Constant, const UNiagaraDataObject* Default)override;
 	virtual TNiagaraExprPtr GetInternalConstant(const FNiagaraVariableInfo& Constant, float Default)override;
 	virtual TNiagaraExprPtr GetInternalConstant(const FNiagaraVariableInfo& Constant, const FVector4& Default)override;
 	virtual TNiagaraExprPtr GetInternalConstant(const FNiagaraVariableInfo& Constant, const FMatrix& Default)override;
-	virtual TNiagaraExprPtr GetExternalCurveConstant(const FNiagaraVariableInfo& Constant)override;
+	virtual TNiagaraExprPtr GetInternalConstant(const FNiagaraVariableInfo& Constant, const UNiagaraDataObject* Default)override;
 
 	virtual bool CheckInputs(FName OpName, TArray<TNiagaraExprPtr>& Inputs)override;
 	virtual bool CheckOutputs(FName OpName, TArray<TNiagaraExprPtr>& Outputs)override;
diff --git a/Engine/Source/Editor/NiagaraEditor/Private/NiagaraEffectEditor.cpp b/Engine/Source/Editor/NiagaraEditor/Private/NiagaraEffectEditor.cpp
--- a/Engine/Source/Editor/NiagaraEditor/Private/NiagaraEffectEditor.cpp
+++ b/Engine/Source/Editor/NiagaraEditor/Private/NiagaraEffectEditor.cpp
@@ -13,6 +13,7 @@
 const FName FNiagaraEffectEditor::UpdateTabId(TEXT("NiagaraEditor_Effect"));
 const FName FNiagaraEffectEditor::ViewportTabID(TEXT("NiagaraEffectEditor_Viewport"));
 const FName FNiagaraEffectEditor::EmitterEditorTabID(TEXT("NiagaraEffectEditor_EmitterEditor"));
+const FName FNiagaraEffectEditor::DevEmitterEditorTabID(TEXT("NiagaraEffectEditor_DevEmitterEditor"));
 const FName FNiagaraEffectEditor::CurveEditorTabID(TEXT("NiagaraEffectEditor_CurveEditor"));
 const FName FNiagaraEffectEditor::SequencerTabID(TEXT("NiagaraEffectEditor_Sequencer"));
 
@@ -31,6 +32,10 @@ void FNiagaraEffectEditor::RegisterTabSpawners(const TSharedRef<class FTabManage
 		.SetDisplayName(LOCTEXT("Emitters", "Emitters"))
 		.SetGroup(WorkspaceMenuCategory.ToSharedRef());
 
+	TabManager->RegisterTabSpawner(DevEmitterEditorTabID, FOnSpawnTab::CreateSP(this, &FNiagaraEffectEditor::SpawnTab_DevEmitterList))
+		.SetDisplayName(LOCTEXT("DevEmitters", "Emitters_Dev"))
+		.SetGroup(WorkspaceMenuCategory.ToSharedRef());
+
 	TabManager->RegisterTabSpawner(CurveEditorTabID, FOnSpawnTab::CreateSP(this, &FNiagaraEffectEditor::SpawnTab_CurveEd))
 		.SetDisplayName(LOCTEXT("Curves", "Curves"))
 		.SetGroup(WorkspaceMenuCategory.ToSharedRef());
@@ -54,6 +59,7 @@ void FNiagaraEffectEditor::UnregisterTabSpawners(const TSharedRef<class FTabMana
 	TabManager->UnregisterTabSpawner(UpdateTabId);
 	TabManager->UnregisterTabSpawner(ViewportTabID);
 	TabManager->UnregisterTabSpawner(EmitterEditorTabID);
+	TabManager->UnregisterTabSpawner(DevEmitterEditorTabID);
 	TabManager->UnregisterTabSpawner(CurveEditorTabID);
 	TabManager->UnregisterTabSpawner(SequencerTabID);
 }
@@ -123,6 +129,7 @@ void FNiagaraEffectEditor::InitNiagaraEffectEditor(const EToolkitMode::Type Mode
 					FTabManager::NewStack()
 					->SetSizeCoefficient(0.3f)
 					->AddTab(EmitterEditorTabID, ETabState::OpenedTab)
+					->AddTab(DevEmitterEditorTabID, ETabState::OpenedTab)
 				)
 			)
 			->Split
@@ -166,6 +173,14 @@ FLinearColor FNiagaraEffectEditor::GetWorldCentricTabColorScale() const
 	return FLinearColor(0.0f, 0.0f, 0.2f, 0.5f);
 }
 
+void FNiagaraEffectEditor::NotifyPostChange(const FPropertyChangedEvent& PropertyChangedEvent, UProperty* PropertyThatChanged)
+{
+	//Rebuild the effect after every edit.
+// 	if (EffectInstance)
+// 	{
+// 		EffectInstance->Init();
+// 	}
+}
 
 /** Create new tab for the supplied graph - don't call this directly, call SExplorer->FindTabForGraph.*/
 TSharedRef<SNiagaraEffectEditorWidget> FNiagaraEffectEditor::CreateEditorWidget(UNiagaraEffect* InEffect)
@@ -202,7 +217,7 @@ TSharedRef<SNiagaraEffectEditorWidget> FNiagaraEffectEditor::CreateEditorWidget(
 		
 		
 	// Make the effect editor widget
-	return SNew(SNiagaraEffectEditorWidget).EffectObj(InEffect).EffectInstance(EffectInstance);
+	return SNew(SNiagaraEffectEditorWidget).EffectObj(InEffect).EffectInstance(EffectInstance).EffectEditor(this);
 }
 
 
@@ -259,6 +274,20 @@ TSharedRef<SDockTab> FNiagaraEffectEditor::SpawnTab_EmitterList(const FSpawnTabA
 	return SpawnedTab;
 }
 
+TSharedRef<SDockTab> FNiagaraEffectEditor::SpawnTab_DevEmitterList(const FSpawnTabArgs& Args)
+{
+	check(Args.GetTabId().TabType == DevEmitterEditorTabID);
+
+	TSharedRef<SDockTab> SpawnedTab =
+		SNew(SDockTab)
+		[
+			SAssignNew(DevEmitterEditorWidget, SNiagaraEffectEditorWidget).EffectInstance(EffectInstance).EffectEditor(this).bForDev(true)
+		];
+
+
+	return SpawnedTab;
+}
+
 
 TSharedRef<SDockTab> FNiagaraEffectEditor::SpawnTab_CurveEd(const FSpawnTabArgs& Args)
 {
@@ -352,23 +381,70 @@ FReply FNiagaraEffectEditor::OnAddEmitterClicked()
 {
 	FNiagaraEditorModule& NiagaraEditorModule = FModuleManager::LoadModuleChecked<FNiagaraEditorModule>("NiagaraEditor");
 
-	FNiagaraEmitterProperties *Props = Effect->AddEmitterProperties();
+	UNiagaraEmitterProperties *Props = Effect->AddEmitterProperties();
 	TSharedPtr<FNiagaraSimulation> NewEmitter = EffectInstance->AddEmitter(Props);
 	Effect->CreateEffectRendererProps(NewEmitter);
 	Viewport->SetPreviewEffect(EffectInstance);
 	EmitterEditorWidget->UpdateList();
+	DevEmitterEditorWidget->UpdateList();
+	
+	Effect->MarkPackageDirty();
+
+	return FReply::Handled();
+}
+
+FReply FNiagaraEffectEditor::OnDuplicateEmitterClicked(TSharedPtr<FNiagaraSimulation> Emitter)
+{
+	FNiagaraEditorModule& NiagaraEditorModule = FModuleManager::LoadModuleChecked<FNiagaraEditorModule>("NiagaraEditor");
+
+	if (UNiagaraEmitterProperties* ToDupe = Emitter->GetProperties().Get())
+	{
+		UNiagaraEmitterProperties *Props = CastChecked<UNiagaraEmitterProperties>(StaticDuplicateObject(ToDupe,Effect,NULL));
+		Effect->AddEmitterProperties(Props);
+		TSharedPtr<FNiagaraSimulation> NewEmitter = EffectInstance->AddEmitter(Props);
+		Effect->CreateEffectRendererProps(NewEmitter);
+		Viewport->SetPreviewEffect(EffectInstance);
+		EmitterEditorWidget->UpdateList();
+		DevEmitterEditorWidget->UpdateList();
+
+		Effect->MarkPackageDirty();
+	}
+	return FReply::Handled();
+}
+
+FReply FNiagaraEffectEditor::OnDeleteEmitterClicked(TSharedPtr<FNiagaraSimulation> Emitter)
+{
+	if (UNiagaraEmitterProperties* ToDelete = Emitter->GetProperties().Get())
+	{
+		Effect->DeleteEmitterProperties(ToDelete);
+		EffectInstance->DeleteEmitter(Emitter);
+		Viewport->SetPreviewEffect(EffectInstance);
+		EmitterEditorWidget->UpdateList();
+		DevEmitterEditorWidget->UpdateList();
+
+		Effect->MarkPackageDirty();
+	}
 	return FReply::Handled();
 }
 
 
 FReply FNiagaraEffectEditor::OnEmitterSelected(TSharedPtr<FNiagaraSimulation> SelectedItem, ESelectInfo::Type SelType)
 {
-	if (SelectedItem.Get() != nullptr)
+ 	if (SelectedItem.Get() != nullptr)
 	{
-		if (SelectedItem->GetProperties()->ExternalConstants.GetDataConstants().Num())
+		if (UNiagaraEmitterProperties* PinnedProps = SelectedItem->GetProperties().Get())
 		{
-			UNiagaraDataObject *DataObj = SelectedItem->GetProperties()->ExternalConstants.GetDataConstants().CreateConstIterator().Value();
-			TimeLine.Get()->SetCurve(static_cast<UNiagaraCurveDataObject*>(DataObj)->GetCurveObject());
+			if (PinnedProps->UpdateScriptProps.ExternalConstants.GetNumDataObjectConstants() > 0)
+			{
+				FNiagaraVariableInfo VarInfo;
+				UNiagaraDataObject* DataObj;
+				PinnedProps->UpdateScriptProps.ExternalConstants.GetDataObjectConstant(0, DataObj, VarInfo);
+
+				if (UNiagaraCurveDataObject* CurvObj = Cast<UNiagaraCurveDataObject>(DataObj))
+				{
+					TimeLine.Get()->SetCurve(CurvObj->GetCurveObject());
+				}
+			}
 		}
 	}
 	return FReply::Handled();
diff --git a/Engine/Source/Editor/NiagaraEditor/Private/NiagaraEffectEditor.h b/Engine/Source/Editor/NiagaraEditor/Private/NiagaraEffectEditor.h
--- a/Engine/Source/Editor/NiagaraEditor/Private/NiagaraEffectEditor.h
+++ b/Engine/Source/Editor/NiagaraEditor/Private/NiagaraEffectEditor.h
@@ -12,7 +12,7 @@
 
 /** Viewer/editor for a NiagaraEffect
 */
-class FNiagaraEffectEditor : public INiagaraEffectEditor
+class FNiagaraEffectEditor : public INiagaraEffectEditor, public FNotifyHook
 {
 
 public:
@@ -44,6 +44,10 @@ public:
 		return MakeShareable(new FNiagaraTrackEditor(InSequencer) );
 	}
 
+	virtual void NotifyPostChange(const FPropertyChangedEvent& PropertyChangedEvent, UProperty* PropertyThatChanged) override;
+
+	FReply OnDeleteEmitterClicked(TSharedPtr<FNiagaraSimulation> Emitter);
+	FReply OnDuplicateEmitterClicked(TSharedPtr<FNiagaraSimulation> Emitter);
 private:
 	/** Create widget for graph editing */
 	TSharedRef<class SNiagaraEffectEditorWidget> CreateEditorWidget(UNiagaraEffect* InEffect);
@@ -53,10 +57,12 @@ private:
 
 	TSharedPtr<SNiagaraEffectEditorViewport>	Viewport;
 	TSharedPtr<SNiagaraEffectEditorWidget>	EmitterEditorWidget;
+	TSharedPtr<SNiagaraEffectEditorWidget>	DevEmitterEditorWidget;
 	TSharedPtr< SNiagaraTimeline > TimeLine;
-
+	
 	TSharedRef<SDockTab> SpawnTab_Viewport(const FSpawnTabArgs& Args);
 	TSharedRef<SDockTab> SpawnTab_EmitterList(const FSpawnTabArgs& Args);
+	TSharedRef<SDockTab> SpawnTab_DevEmitterList(const FSpawnTabArgs& Args);
 	TSharedRef<SDockTab> SpawnTab_CurveEd(const FSpawnTabArgs& Args);
 	TSharedRef<SDockTab> SpawnTab_Sequencer(const FSpawnTabArgs& Args);
 
@@ -85,6 +91,7 @@ private:
 	static const FName UpdateTabId;
 	static const FName ViewportTabID;
 	static const FName EmitterEditorTabID;
+	static const FName DevEmitterEditorTabID;
 	static const FName CurveEditorTabID;
 	static const FName SequencerTabID;
 };
diff --git a/Engine/Source/Editor/NiagaraEditor/Private/NiagaraEmitterPropertiesDetailsCustomization.cpp b/Engine/Source/Editor/NiagaraEditor/Private/NiagaraEmitterPropertiesDetailsCustomization.cpp
new file mode 100644
--- /dev/null
+++ b/Engine/Source/Editor/NiagaraEditor/Private/NiagaraEmitterPropertiesDetailsCustomization.cpp
@@ -0,0 +1,161 @@
+// Copyright 1998-2015 Epic Games, Inc. All Rights Reserved.
+
+#include "NiagaraEditorPrivatePCH.h"
+#include "NiagaraEditorModule.h"
+#include "EditorSupportDelegates.h"
+#include "UnrealEd.h"
+#include "PropertyEditing.h"
+#include "PropertyCustomizationHelpers.h"
+#include "NiagaraEffectEditor.h"
+#include "NiagaraEmitterPropertiesDetailsCustomization.h"
+#include "ScopedTransaction.h"
+#include "DetailWidgetRow.h"
+#include "AssetData.h"
+
+#define LOCTEXT_NAMESPACE "NiagaraEffectEditor"
+
+TSharedRef<IDetailCustomization> FNiagaraEmitterPropertiesDetails::MakeInstance(TWeakObjectPtr<UNiagaraEmitterProperties> EmitterProperties)
+{
+	return MakeShareable(new FNiagaraEmitterPropertiesDetails(EmitterProperties));
+}
+
+FNiagaraEmitterPropertiesDetails::FNiagaraEmitterPropertiesDetails(TWeakObjectPtr<UNiagaraEmitterProperties> EmitterProperties)
+: EmitterProps(EmitterProperties)
+{
+}
+
+void FNiagaraEmitterPropertiesDetails::CustomizeDetails(IDetailLayoutBuilder& DetailLayout)
+{
+	FName EmitterCategoryName = TEXT("Emitter");
+	IDetailCategoryBuilder& EmitterCategory = DetailLayout.EditCategory(EmitterCategoryName);
+
+	TSharedRef<IPropertyHandle> EmitterName = DetailLayout.GetProperty(GET_MEMBER_NAME_CHECKED(UNiagaraEmitterProperties, EmitterName));
+	TSharedRef<IPropertyHandle> EmitterEnabled = DetailLayout.GetProperty(GET_MEMBER_NAME_CHECKED(UNiagaraEmitterProperties, bIsEnabled));
+
+	DetailLayout.HideProperty(EmitterName);
+	DetailLayout.HideProperty(EmitterEnabled);
+	
+	TSharedRef<IPropertyHandle> UpdateScriptProps = DetailLayout.GetProperty(GET_MEMBER_NAME_CHECKED(UNiagaraEmitterProperties, UpdateScriptProps));
+	TSharedRef<IPropertyHandle> SpawnScriptProps = DetailLayout.GetProperty(GET_MEMBER_NAME_CHECKED(UNiagaraEmitterProperties, SpawnScriptProps));
+
+	BuildScriptProperties(DetailLayout, SpawnScriptProps, TEXT("Spawn Script"), LOCTEXT("SpawnScriptDisplayName", "Spawn Script"));
+	BuildScriptProperties(DetailLayout, UpdateScriptProps, TEXT("Update Script"), LOCTEXT("UpdateScriptDisplayName", "Update Script"));
+}
+
+void FNiagaraEmitterPropertiesDetails::BuildScriptProperties(IDetailLayoutBuilder& DetailLayout, TSharedRef<IPropertyHandle> ScriptPropsHandle, FName Name, FText DisplayName)
+{
+	TFunction<void(TSharedRef<IPropertyHandle>)> HideAllBelow = [&](TSharedRef<IPropertyHandle> Prop)
+	{
+		DetailLayout.HideProperty(Prop);
+
+		//Don't hide sub properties for data objects.
+		if (Prop->GetProperty()->GetName() == GET_MEMBER_NAME_CHECKED(FNiagaraConstants, DataObjectConstants).ToString())
+			return; 
+
+		uint32 NumChildren = 0;
+		if (Prop->GetNumChildren(NumChildren) == FPropertyAccess::Success)
+		{
+			for (uint32 i = 0; i < NumChildren; ++i)
+			{
+				TSharedPtr<IPropertyHandle> Child = Prop->GetChildHandle(i);
+				if (Child.IsValid())
+				{
+					HideAllBelow(Child.ToSharedRef());
+				}				
+			}
+		}
+	};
+	HideAllBelow(ScriptPropsHandle);
+
+	IDetailCategoryBuilder& ScriptCategory = DetailLayout.EditCategory(Name, DisplayName);
+
+	//Script
+	TSharedPtr<IPropertyHandle> ScriptHandle = ScriptPropsHandle->GetChildHandle(GET_MEMBER_NAME_CHECKED(FNiagaraEmitterScriptProperties, Script));
+	ScriptCategory.AddProperty(ScriptHandle);
+
+	//Constants
+	const bool bGenerateHeader = false;
+	const bool bDisplayResetToDefault = false;
+	TSharedPtr<IPropertyHandle> Constants = ScriptPropsHandle->GetChildHandle(GET_MEMBER_NAME_CHECKED(FNiagaraEmitterScriptProperties, ExternalConstants));
+	//Scalar Constants.
+	TSharedPtr<IPropertyHandle> ScalarConstants = ScriptPropsHandle->GetChildHandle(GET_MEMBER_NAME_CHECKED(FNiagaraConstants, ScalarConstants));
+	TSharedRef<FDetailArrayBuilder> ScalarConstantsBuilder = MakeShareable(new FDetailArrayBuilder(ScalarConstants.ToSharedRef(), bGenerateHeader, bDisplayResetToDefault));
+	ScalarConstantsBuilder->OnGenerateArrayElementWidget(FOnGenerateArrayElementWidget::CreateSP(this, &FNiagaraEmitterPropertiesDetails::OnGenerateScalarConstantEntry));
+	ScriptCategory.AddCustomBuilder(ScalarConstantsBuilder);
+	//Vector Constants.
+	TSharedPtr<IPropertyHandle> VectorConstants = ScriptPropsHandle->GetChildHandle(GET_MEMBER_NAME_CHECKED(FNiagaraConstants, VectorConstants));
+	TSharedRef<FDetailArrayBuilder> VectorConstantsBuilder = MakeShareable(new FDetailArrayBuilder(VectorConstants.ToSharedRef(), bGenerateHeader, bDisplayResetToDefault));
+	VectorConstantsBuilder->OnGenerateArrayElementWidget(FOnGenerateArrayElementWidget::CreateSP(this, &FNiagaraEmitterPropertiesDetails::OnGenerateVectorConstantEntry));
+	ScriptCategory.AddCustomBuilder(VectorConstantsBuilder);
+	//Matrix Constants.
+	TSharedPtr<IPropertyHandle> MatrixConstants = ScriptPropsHandle->GetChildHandle(GET_MEMBER_NAME_CHECKED(FNiagaraConstants, MatrixConstants));
+	TSharedRef<FDetailArrayBuilder> MatrixConstantsBuilder = MakeShareable(new FDetailArrayBuilder(MatrixConstants.ToSharedRef(), bGenerateHeader, bDisplayResetToDefault));
+	MatrixConstantsBuilder->OnGenerateArrayElementWidget(FOnGenerateArrayElementWidget::CreateSP(this, &FNiagaraEmitterPropertiesDetails::OnGenerateMatrixConstantEntry));
+	ScriptCategory.AddCustomBuilder(MatrixConstantsBuilder);
+	//DataObject Constants.
+	TSharedPtr<IPropertyHandle> DataObjectConstants = ScriptPropsHandle->GetChildHandle(GET_MEMBER_NAME_CHECKED(FNiagaraConstants, DataObjectConstants));
+	TSharedRef<FDetailArrayBuilder> DataObjectConstantsBuilder = MakeShareable(new FDetailArrayBuilder(DataObjectConstants.ToSharedRef(), bGenerateHeader, bDisplayResetToDefault));
+	DataObjectConstantsBuilder->OnGenerateArrayElementWidget(FOnGenerateArrayElementWidget::CreateSP(this, &FNiagaraEmitterPropertiesDetails::OnGenerateDataObjectConstantEntry));
+	ScriptCategory.AddCustomBuilder(DataObjectConstantsBuilder);
+
+	//TODO: Events.
+}
+
+void FNiagaraEmitterPropertiesDetails::OnGenerateScalarConstantEntry(TSharedRef<IPropertyHandle> ElementProperty, int32 ElementIndex, IDetailChildrenBuilder& ChildrenBuilder)
+{
+	TSharedPtr<IPropertyHandle> ValueProperty = ElementProperty->GetChildHandle(GET_MEMBER_NAME_CHECKED(FNiagaraConstants_Float, Value));
+	TSharedPtr<IPropertyHandle> NameProperty = ElementProperty->GetChildHandle(GET_MEMBER_NAME_CHECKED(FNiagaraConstants_Float, Name));
+
+	FName DisplayName;
+	NameProperty->GetValue(DisplayName);
+	//Don't display system constants
+	if (UNiagaraComponent::GetSystemConstants().Find(FNiagaraVariableInfo(DisplayName, ENiagaraDataType::Scalar)) == INDEX_NONE)
+	{
+		ChildrenBuilder.AddChildProperty(ValueProperty.ToSharedRef()).DisplayName(FText::FromName(DisplayName));
+	}
+}
+
+void FNiagaraEmitterPropertiesDetails::OnGenerateVectorConstantEntry(TSharedRef<IPropertyHandle> ElementProperty, int32 ElementIndex, IDetailChildrenBuilder& ChildrenBuilder)
+{
+	TSharedPtr<IPropertyHandle> ValueProperty = ElementProperty->GetChildHandle(GET_MEMBER_NAME_CHECKED(FNiagaraConstants_Vector, Value));
+	TSharedPtr<IPropertyHandle> NameProperty = ElementProperty->GetChildHandle(GET_MEMBER_NAME_CHECKED(FNiagaraConstants_Vector, Name));
+
+	FName DisplayName;
+	NameProperty->GetValue(DisplayName);
+	//Don't display system constants
+	if (UNiagaraComponent::GetSystemConstants().Find(FNiagaraVariableInfo(DisplayName, ENiagaraDataType::Vector)) == INDEX_NONE)
+	{
+		ChildrenBuilder.AddChildProperty(ValueProperty.ToSharedRef()).DisplayName(FText::FromName(DisplayName));
+	}
+}
+
+void FNiagaraEmitterPropertiesDetails::OnGenerateMatrixConstantEntry(TSharedRef<IPropertyHandle> ElementProperty, int32 ElementIndex, IDetailChildrenBuilder& ChildrenBuilder)
+{
+	TSharedPtr<IPropertyHandle> ValueProperty = ElementProperty->GetChildHandle(GET_MEMBER_NAME_CHECKED(FNiagaraConstants_Matrix, Value));
+	TSharedPtr<IPropertyHandle> NameProperty = ElementProperty->GetChildHandle(GET_MEMBER_NAME_CHECKED(FNiagaraConstants_Matrix, Name));
+
+	FName DisplayName;
+	NameProperty->GetValue(DisplayName);
+	//Don't display system constants
+	if (UNiagaraComponent::GetSystemConstants().Find(FNiagaraVariableInfo(DisplayName, ENiagaraDataType::Matrix)) == INDEX_NONE)
+	{
+		ChildrenBuilder.AddChildProperty(ValueProperty.ToSharedRef()).DisplayName(FText::FromName(DisplayName));
+	}
+}
+
+void FNiagaraEmitterPropertiesDetails::OnGenerateDataObjectConstantEntry(TSharedRef<IPropertyHandle> ElementProperty, int32 ElementIndex, IDetailChildrenBuilder& ChildrenBuilder)
+{
+	TSharedPtr<IPropertyHandle> ValueProperty = ElementProperty->GetChildHandle(GET_MEMBER_NAME_CHECKED(FNiagaraConstants_DataObject, Value));
+	TSharedPtr<IPropertyHandle> NameProperty = ElementProperty->GetChildHandle(GET_MEMBER_NAME_CHECKED(FNiagaraConstants_DataObject, Name));
+
+	FName DisplayName;
+	NameProperty->GetValue(DisplayName);
+	//Don't display system constants
+	if (UNiagaraComponent::GetSystemConstants().Find(FNiagaraVariableInfo(DisplayName, ENiagaraDataType::Curve)) == INDEX_NONE)
+	{
+		ChildrenBuilder.AddChildProperty(ValueProperty.ToSharedRef()).DisplayName(FText::FromName(DisplayName));
+	}
+}
+
+#undef LOCTEXT_NAMESPACE
+
diff --git a/Engine/Source/Editor/NiagaraEditor/Private/NiagaraEmitterPropertiesDetailsCustomization.h b/Engine/Source/Editor/NiagaraEditor/Private/NiagaraEmitterPropertiesDetailsCustomization.h
new file mode 100644
--- /dev/null
+++ b/Engine/Source/Editor/NiagaraEditor/Private/NiagaraEmitterPropertiesDetailsCustomization.h
@@ -0,0 +1,31 @@
+// Copyright 1998-2015 Epic Games, Inc. All Rights Reserved.
+
+#pragma once
+
+/*-----------------------------------------------------------------------------
+FNiagaraEmitterPropertiesDetails
+-----------------------------------------------------------------------------*/
+
+class FNiagaraEmitterPropertiesDetails : public IDetailCustomization
+{
+public:
+	/** Makes a new instance of this detail layout class for a specific detail view requesting it */
+	static TSharedRef<class IDetailCustomization> MakeInstance(TWeakObjectPtr<UNiagaraEmitterProperties> EmitterProperties);
+
+	/** Constructor */
+	FNiagaraEmitterPropertiesDetails(TWeakObjectPtr<UNiagaraEmitterProperties> EmitterProperties);
+
+	/** IDetailCustomization interface */
+	virtual void CustomizeDetails(IDetailLayoutBuilder& DetailLayout) override;
+
+	void OnGenerateScalarConstantEntry(TSharedRef<IPropertyHandle> ElementProperty, int32 ElementIndex, IDetailChildrenBuilder& ChildrenBuilder);
+	void OnGenerateVectorConstantEntry(TSharedRef<IPropertyHandle> ElementProperty, int32 ElementIndex, IDetailChildrenBuilder& ChildrenBuilder);
+	void OnGenerateMatrixConstantEntry(TSharedRef<IPropertyHandle> ElementProperty, int32 ElementIndex, IDetailChildrenBuilder& ChildrenBuilder);
+	void OnGenerateDataObjectConstantEntry(TSharedRef<IPropertyHandle> ElementProperty, int32 ElementIndex, IDetailChildrenBuilder& ChildrenBuilder);
+
+	void BuildScriptProperties(IDetailLayoutBuilder& DetailLayout, TSharedRef<IPropertyHandle> ScriptPropsHandle, FName Name, FText DisplayName);
+private:
+	/** Object that stores all of the possible parameters we can edit */
+	TWeakObjectPtr<UNiagaraEmitterProperties> EmitterProps;
+};
+
diff --git a/Engine/Source/Editor/NiagaraEditor/Private/NiagaraSequencer.h b/Engine/Source/Editor/NiagaraEditor/Private/NiagaraSequencer.h
--- a/Engine/Source/Editor/NiagaraEditor/Private/NiagaraSequencer.h
+++ b/Engine/Source/Editor/NiagaraEditor/Private/NiagaraSequencer.h
@@ -22,12 +22,12 @@ class UNiagaraMovieSceneSection : public UMovieSceneSection
 public:
 	FText GetEmitterName()	{ return EmitterName; }
 	void SetEmitterName(FText InName)	{ EmitterName = InName; }
-	void SetEmitterProps(const FNiagaraEmitterProperties *Props)	{ EmitterProps = Props; }
+	void SetEmitterProps(TWeakObjectPtr<const UNiagaraEmitterProperties> Props)	{ EmitterProps = Props; }
 
-	const FNiagaraEmitterProperties *GetEmitterProps() const	{ return EmitterProps; }
+	TWeakObjectPtr<const UNiagaraEmitterProperties> GetEmitterProps() const	{ return EmitterProps; }
 private:
 	FText EmitterName;
-	const FNiagaraEmitterProperties *EmitterProps;
+	TWeakObjectPtr<const UNiagaraEmitterProperties> EmitterProps;
 };
 
 /**
@@ -139,20 +139,24 @@ class UEmitterMovieSceneTrack : public UMovieSceneTrack
 public:
 
 	/** UMovieSceneTrack interface */
-	virtual FName GetTrackName() const override { return *Emitter->GetProperties()->Name; }
+	virtual FName GetTrackName() const override { return *Emitter->GetProperties()->EmitterName; }
 	void SetEmitter(TSharedPtr<FNiagaraSimulation> InEmitter)
 	{
-		Emitter = InEmitter;
-		const FNiagaraEmitterProperties *EmitterProps = InEmitter->GetProperties();
-
-		UNiagaraMovieSceneSection *Section = NewObject<UNiagaraMovieSceneSection>(this, *Emitter->GetProperties()->Name);
-		Section->SetEmitterProps(EmitterProps);
-
-		Section->SetStartTime(EmitterProps->StartTime);
-		Section->SetEndTime(EmitterProps->EndTime);
-		Section->SetEmitterName(FText::FromString(EmitterProps->Name));
-
-		Sections.Add(Section);
+		if (InEmitter.IsValid())
+		{
+			Emitter = InEmitter;
+			if (const UNiagaraEmitterProperties* EmitterProps = InEmitter->GetProperties().Get())
+			{
+				UNiagaraMovieSceneSection *Section = NewObject<UNiagaraMovieSceneSection>(this, *Emitter->GetProperties()->EmitterName);
+				Section->SetEmitterProps(EmitterProps);
+
+				Section->SetStartTime(EmitterProps->StartTime);
+				Section->SetEndTime(EmitterProps->EndTime);
+				Section->SetEmitterName(FText::FromString(EmitterProps->EmitterName));
+
+				Sections.Add(Section);
+			}
+		}
 	}
 
 	virtual void RemoveAllAnimationData() override	{};
diff --git a/Engine/Source/Editor/NiagaraEditor/Private/SNiagaraEffectEditorWidget.cpp b/Engine/Source/Editor/NiagaraEditor/Private/SNiagaraEffectEditorWidget.cpp
--- a/Engine/Source/Editor/NiagaraEditor/Private/SNiagaraEffectEditorWidget.cpp
+++ b/Engine/Source/Editor/NiagaraEditor/Private/SNiagaraEffectEditorWidget.cpp
@@ -8,7 +8,7 @@
 #include "Editor/PropertyEditor/Public/PropertyEditorModule.h"
 #include "Editor/PropertyEditor/Public/IDetailsView.h"
 #include "NiagaraEffectEditor.h"
-
+#include "NiagaraEmitterPropertiesDetailsCustomization.h"
 
 #define LOCTEXT_NAMESPACE "NiagaraEffectEditor"
 
@@ -27,43 +27,82 @@ void SNiagaraEffectEditorWidget::Construct(const FArguments& InArgs)
 	}
 
 	EffectEditor = InArgs._EffectEditor;
-
-	TSharedPtr<SOverlay> OverlayWidget;
-	this->ChildSlot
-		[
-
-			SNew(SSplitter)
-			.Orientation(Orient_Vertical)
-			+ SSplitter::Slot()
-			.Value(0.75f)
+	bForDev = InArgs._bForDev;
+	
+	if (bForDev)
+	{
+		ChildSlot
+			.Padding(4)
 			[
+				SNew(SOverlay)
+				+ SOverlay::Slot()
+				.VAlign(VAlign_Fill)
+				[
+					SAssignNew(EmitterBox, SScrollBox)
+					.Orientation(Orient_Horizontal)
+				]
+			];
 
-				SNew(SVerticalBox)
-				+ SVerticalBox::Slot()
-				.AutoHeight()
+		ConstructDevEmittersView();
+	}
+	else
+	{
+		TSharedPtr<SOverlay> OverlayWidget;
+		this->ChildSlot
+			.Padding(8)
+			[
+				SNew(SSplitter)
+				.Orientation(Orient_Vertical)
+				+ SSplitter::Slot()
+				.Value(0.75f)
 				[
-					SNew(SHorizontalBox)
-					+ SHorizontalBox::Slot()
-						.AutoWidth()
+					SNew(SScrollBox)
+					+ SScrollBox::Slot()
+					.Padding(0.0f)
+					[
+						SNew(SVerticalBox)
+						+ SVerticalBox::Slot()
+						.AutoHeight()
 						[
-							SAssignNew(OverlayWidget, SOverlay)
-							+SOverlay::Slot()
-							.Padding(10)
-							.VAlign(VAlign_Fill)
+							SNew(SHorizontalBox)
+							+ SHorizontalBox::Slot()
+							.AutoWidth()
 							[
-								SAssignNew(ListView, SListView<TSharedPtr<FNiagaraSimulation> >)
-								.ItemHeight(256)
-								.ListItemsSource(&(this->EffectInstance->GetEmitters()))
-								.OnGenerateRow(this, &SNiagaraEffectEditorWidget::OnGenerateWidgetForList)
-								.OnSelectionChanged(this, &SNiagaraEffectEditorWidget::OnSelectionChanged)
-								//.SelectionMode(ESelectionMode::Single)
+								SAssignNew(OverlayWidget, SOverlay)
+								+ SOverlay::Slot()
+								.Padding(10)
+								.VAlign(VAlign_Fill)
+								[
+									SAssignNew(ListView, SListView<TSharedPtr<FNiagaraSimulation> >)
+									.ItemHeight(256)
+									.ListItemsSource(&(this->EffectInstance->GetEmitters()))
+									.OnGenerateRow(this, &SNiagaraEffectEditorWidget::OnGenerateWidgetForList)
+									.OnSelectionChanged(this, &SNiagaraEffectEditorWidget::OnSelectionChanged)
+									//.SelectionMode(ESelectionMode::Single)
+								]
 							]
 						]
+					]
 				]
-			]
-		];
+			];
+	}
 }
 
+void SNiagaraEffectEditorWidget::ConstructDevEmittersView()
+{
+	EmitterBox->ClearChildren();
+	for (TSharedPtr<FNiagaraSimulation>& Emitter : EffectInstance->GetEmitters())
+	{
+		EmitterBox->AddSlot()
+		.Padding(2.0f)
+			[
+				NGED_SECTION_DARKBORDER
+				[
+					SNew(SEmitterWidgetDev).Emitter(Emitter).Effect(EffectInstance).EffectEditor(EffectEditor)
+				]
+			];
+	}
+}
 
 void SNiagaraTimeline::Construct(const FArguments& InArgs)
 {
@@ -78,13 +117,137 @@ void SNiagaraTimeline::Construct(const FArguments& InArgs)
 	];
 }
 
-
-void SEmitterWidget::Construct(const FArguments& InArgs)
+void SEmitterWidgetBase::Construct(const FArguments& InArgs)
 {
-	CurSpawnScript = nullptr;
-	CurUpdateScript = nullptr;
 	Emitter = InArgs._Emitter;
 	EffectInstance = InArgs._Effect;
+	EffectEditor = InArgs._EffectEditor;
+}
+
+void SEmitterWidgetDev::Construct(const FArguments& InArgs)
+{
+	SEmitterWidgetBase::Construct(InArgs);
+
+	BuildContents();
+}
+
+void SEmitterWidgetDev::BuildContents()
+{
+	//Register details customization.	
+	if (UNiagaraEmitterProperties* PinnedProps = Emitter->GetProperties().Get())
+	{
+		FPropertyEditorModule& PropertyEditorModule = FModuleManager::GetModuleChecked<FPropertyEditorModule>("PropertyEditor");
+		FDetailsViewArgs DetailsViewArgs(false, false, true, FDetailsViewArgs::HideNameArea, true, this);
+		Details = PropertyEditorModule.CreateDetailView(DetailsViewArgs);
+		FOnGetDetailCustomizationInstance LayoutEmitterDetails = FOnGetDetailCustomizationInstance::CreateStatic(&FNiagaraEmitterPropertiesDetails::MakeInstance, Emitter->GetProperties());
+		Details->RegisterInstancedCustomPropertyLayout(UNiagaraEmitterProperties::StaticClass(), LayoutEmitterDetails);
+
+		Details->SetObject(PinnedProps);
+
+		ChildSlot
+			.Padding(4)
+			[
+				SNew(SVerticalBox)
+				+ SVerticalBox::Slot()
+				.AutoHeight()
+				.Padding(0)
+				[
+					NGED_SECTION_BORDER
+					[
+						SNew(SVerticalBox)
+						+ SVerticalBox::Slot()
+						.AutoHeight()
+						.HAlign(HAlign_Fill)
+						[
+							// name and status line
+							SNew(SHorizontalBox)
+							+ SHorizontalBox::Slot()
+							.HAlign(HAlign_Left)
+							.AutoWidth()
+							.Padding(2)
+							[
+								SNew(SCheckBox)
+								.OnCheckStateChanged(this, &SEmitterWidgetBase::OnEmitterEnabledChanged)
+								.IsChecked(this, &SEmitterWidgetBase::IsEmitterEnabled)
+								.ToolTipText(FText::FromString("Toggles whether this emitter is enabled. Disabled emitters don't simulate or render."))
+							]
+							+ SHorizontalBox::Slot()
+								.AutoWidth()
+								.Padding(2)
+								.HAlign(HAlign_Left)
+								[
+									SNew(SEditableText).Text(this, &SEmitterWidgetBase::GetEmitterName)
+									.MinDesiredWidth(200)
+									.Font(FEditorStyle::GetFontStyle("ContentBrowser.AssetListViewNameFontDirty"))
+									.OnTextChanged(this, &SEmitterWidgetBase::OnEmitterNameChanged)
+								]
+							+ SHorizontalBox::Slot()
+								.HAlign(HAlign_Right)
+								.AutoWidth()
+								[
+									SNew(SButton)
+									.VAlign(VAlign_Center)
+									.HAlign(HAlign_Right)
+									.Text(LOCTEXT("DeleteEmitter", "Delete"))
+									.ToolTipText(LOCTEXT("DeleteEmitter_Tooltip", "Deletes this emitter from the effect."))
+									.OnClicked(EffectEditor, &FNiagaraEffectEditor::OnDeleteEmitterClicked, Emitter)
+								]
+							+ SHorizontalBox::Slot()
+								.HAlign(HAlign_Right)
+								.AutoWidth()
+								[
+									SNew(SButton)
+									.VAlign(VAlign_Center)
+									.HAlign(HAlign_Right)
+									.Text(LOCTEXT("DuplicateEmitter", "Duplicate"))
+									.ToolTipText(LOCTEXT("DuplicateEmitter_Tooltip", "Duplicate this emitter."))
+									.OnClicked(EffectEditor, &FNiagaraEffectEditor::OnDuplicateEmitterClicked, Emitter)
+								]
+						]
+						+ SVerticalBox::Slot()
+							.AutoHeight()
+							.HAlign(HAlign_Fill)
+							.Padding(2)
+							[
+								SNew(STextBlock).Text(this, &SEmitterWidgetBase::GetStatsText)
+							]
+					]
+				]
+				+ SVerticalBox::Slot()
+					.VAlign(VAlign_Fill)
+					.Padding(0)
+					[
+						SNew(SHorizontalBox)
+						+ SHorizontalBox::Slot()						
+						.MaxWidth(400)
+						[
+							Details.ToSharedRef()
+						]
+					]
+			];
+	}
+
+}
+
+void SEmitterWidgetDev::NotifyPostChange(const FPropertyChangedEvent& PropertyChangedEvent, UProperty* PropertyThatChanged)
+{
+	//Changes to some properties will require the properties to be reinitialized and the UI to be rebuilt.
+	if (PropertyThatChanged && PropertyThatChanged->GetName() == TEXT("Script"))
+	{
+		Emitter->GetProperties()->Init();
+
+		ChildSlot.DetachWidget();
+		BuildContents();
+	}
+
+	//Always refresh the emitter on a propety change.
+	Emitter->Init();
+}
+
+void SEmitterWidget::Construct(const FArguments& InArgs)
+{
+	SEmitterWidgetBase::Construct(InArgs);
+
 	CurMaterial = nullptr;
 
 	ThumbnailPool = MakeShareable(new FAssetThumbnailPool(1, /*InAreRealTileThumbnailsAllowed=*/false));
@@ -102,18 +265,18 @@ void SEmitterWidget::Construct(const FArguments& InArgs)
 	static TArray<TSharedPtr<EditorExposedVectorCurveConstant>> DummyCurveArray;
 	TArray<TSharedPtr<EditorExposedVectorConstant>> *EditorExposedConstants = &DummyArray;
 	TArray<TSharedPtr<EditorExposedVectorCurveConstant>> *EditorExposedCurveConstants = &DummyCurveArray;
-	if (Emitter->GetProperties()->UpdateScript && Emitter->GetProperties()->UpdateScript->Source)
+	if (Emitter->GetProperties()->UpdateScriptProps.Script && Emitter->GetProperties()->UpdateScriptProps.Script->Source)
 	{
-		EditorExposedConstants = &Emitter->GetProperties()->UpdateScript->Source->ExposedVectorConstants;
-		EditorExposedCurveConstants = &Emitter->GetProperties()->UpdateScript->Source->ExposedVectorCurveConstants;
+		EditorExposedConstants = &Emitter->GetProperties()->UpdateScriptProps.Script->Source->ExposedVectorConstants;
+		EditorExposedCurveConstants = &Emitter->GetProperties()->UpdateScriptProps.Script->Source->ExposedVectorCurveConstants;
 	}
 
 	TArray<TSharedPtr<EditorExposedVectorConstant>> *EditorExposedSpawnConstants = &DummyArray;
 	TArray<TSharedPtr<EditorExposedVectorCurveConstant>> *EditorExposedSpawnCurveConstants = &DummyCurveArray;
-	if (Emitter->GetProperties()->SpawnScript && Emitter->GetProperties()->SpawnScript->Source)
+	if (Emitter->GetProperties()->SpawnScriptProps.Script && Emitter->GetProperties()->SpawnScriptProps.Script->Source)
 	{
-		EditorExposedSpawnConstants = &Emitter->GetProperties()->SpawnScript->Source->ExposedVectorConstants;
-		EditorExposedSpawnCurveConstants = &Emitter->GetProperties()->SpawnScript->Source->ExposedVectorCurveConstants;
+		EditorExposedSpawnConstants = &Emitter->GetProperties()->SpawnScriptProps.Script->Source->ExposedVectorConstants;
+		EditorExposedSpawnCurveConstants = &Emitter->GetProperties()->SpawnScriptProps.Script->Source->ExposedVectorCurveConstants;
 	}
 
 
@@ -148,26 +311,48 @@ void SEmitterWidget::Construct(const FArguments& InArgs)
 						.AutoWidth()
 						[
 							SNew(SCheckBox)
-							.OnCheckStateChanged(this, &SEmitterWidget::OnEmitterEnabledChanged)
-							.IsChecked(this, &SEmitterWidget::IsEmitterEnabled)
+							.OnCheckStateChanged(this, &SEmitterWidgetBase::OnEmitterEnabledChanged)
+							.IsChecked(this, &SEmitterWidgetBase::IsEmitterEnabled)
 							.ToolTipText(FText::FromString("Toggles whether this emitter is enabled. Disabled emitters don't simulate or render."))
 						]
 						+ SHorizontalBox::Slot()
 							.FillWidth(0.3)
 							.HAlign(HAlign_Left)
 							[
-								SNew(SEditableText).Text(this, &SEmitterWidget::GetEmitterName)
+								SNew(SEditableText).Text(this, &SEmitterWidgetBase::GetEmitterName)
 								.MinDesiredWidth(250)
 								.Font(FEditorStyle::GetFontStyle("ContentBrowser.AssetListViewNameFontDirty"))
-								.OnTextChanged(this, &SEmitterWidget::OnEmitterNameChanged)
+								.OnTextChanged(this, &SEmitterWidgetBase::OnEmitterNameChanged)
 							]
 						+ SHorizontalBox::Slot()
 							.HAlign(HAlign_Right)
 							.FillWidth(0.7)
 							[
-								SNew(STextBlock).Text(this, &SEmitterWidget::GetStatsText)
+								SNew(STextBlock).Text(this, &SEmitterWidgetBase::GetStatsText)
 								.MinDesiredWidth(500)
 							]
+							+ SHorizontalBox::Slot()
+								.HAlign(HAlign_Right)
+								.AutoWidth()
+								[
+									SNew(SButton)
+									.VAlign(VAlign_Center)
+									.HAlign(HAlign_Right)
+									.Text(LOCTEXT("DeleteEmitter", "Delete"))
+									.ToolTipText(LOCTEXT("DeleteEmitter_Tooltip", "Deletes this emitter from the effect."))
+									.OnClicked(EffectEditor, &FNiagaraEffectEditor::OnDeleteEmitterClicked, Emitter)
+								]
+							+ SHorizontalBox::Slot()
+								.HAlign(HAlign_Right)
+								.AutoWidth()
+								[
+									SNew(SButton)
+									.VAlign(VAlign_Center)
+									.HAlign(HAlign_Right)
+									.Text(LOCTEXT("DuplicateEmitter", "Duplicate"))
+									.ToolTipText(LOCTEXT("DuplicateEmitter_Tooltip", "Duplicate this emitter."))
+									.OnClicked(EffectEditor, &FNiagaraEffectEditor::OnDuplicateEmitterClicked, Emitter)
+								]
 					]
 				]
 
@@ -505,19 +690,25 @@ void SEmitterWidget::OnUpdateScriptSelectedFromPicker(UObject *Asset)
 	//AssetPickerAnchor->SetIsOpen(false);
 
 	// Set the object found from the asset picker
-	CurUpdateScript = Cast<UNiagaraScript>(Asset);
-	Emitter->SetUpdateScript(CurUpdateScript);
+	UNiagaraScript* CurUpdateScript = Cast<UNiagaraScript>(Asset);
+	UNiagaraEmitterProperties* PinnedProps = Emitter->GetProperties().Get();
+	if (PinnedProps)
+	{
+		PinnedProps->UpdateScriptProps.Script = CurUpdateScript;
+		PinnedProps->UpdateScriptProps.Init(PinnedProps);
+		Emitter->Init();
 
-	UpdateScriptConstantListSlot->DetachWidget();
-	(*UpdateScriptConstantListSlot)
-		[
-			SNew(SExpandableArea).InitiallyCollapsed(true).AreaTitle(LOCTEXT("ConstantsLabel", "Constants"))
-			.BodyContent()
+		UpdateScriptConstantListSlot->DetachWidget();
+		(*UpdateScriptConstantListSlot)
 			[
-				SAssignNew(UpdateScriptConstantList, SListView<TSharedPtr<EditorExposedVectorConstant>>)
-				.ItemHeight(20).ListItemsSource(&CurUpdateScript->Source->ExposedVectorConstants).OnGenerateRow(this, &SEmitterWidget::OnGenerateConstantListRow)
-			]
-		];
+				SNew(SExpandableArea).InitiallyCollapsed(true).AreaTitle(LOCTEXT("ConstantsLabel", "Constants"))
+				.BodyContent()
+				[
+					SAssignNew(UpdateScriptConstantList, SListView<TSharedPtr<EditorExposedVectorConstant>>)
+					.ItemHeight(20).ListItemsSource(&CurUpdateScript->Source->ExposedVectorConstants).OnGenerateRow(this, &SEmitterWidget::OnGenerateConstantListRow)
+				]
+			];
+	}
 }
 
 void SEmitterWidget::OnSpawnScriptSelectedFromPicker(UObject *Asset)
@@ -526,19 +717,26 @@ void SEmitterWidget::OnSpawnScriptSelectedFromPicker(UObject *Asset)
 	//AssetPickerAnchor->SetIsOpen(false);
 
 	// Set the object found from the asset picker
-	CurSpawnScript = Cast<UNiagaraScript>(Asset);
-	Emitter->SetSpawnScript(CurSpawnScript);
+	UNiagaraScript* CurSpawnScript = Cast<UNiagaraScript>(Asset);
+	UNiagaraEmitterProperties* PinnedProps = Emitter->GetProperties().Get();
+	if (PinnedProps)
+	{
+		PinnedProps->SpawnScriptProps.Script = CurSpawnScript;
+		PinnedProps->SpawnScriptProps.Init(PinnedProps);
+		Emitter->Init();
 
-	SpawnScriptConstantListSlot->DetachWidget();
-	(*SpawnScriptConstantListSlot)
-		[
-			SNew(SExpandableArea).InitiallyCollapsed(true).AreaTitle(LOCTEXT("ConstantsLabel", "Constants"))
-			.BodyContent()
+
+		SpawnScriptConstantListSlot->DetachWidget();
+		(*SpawnScriptConstantListSlot)
 			[
-				SAssignNew(SpawnScriptConstantList, SListView<TSharedPtr<EditorExposedVectorConstant>>)
-				.ItemHeight(20).ListItemsSource(&CurSpawnScript->Source->ExposedVectorConstants).OnGenerateRow(this, &SEmitterWidget::OnGenerateConstantListRow)
-			]
-		];
+				SNew(SExpandableArea).InitiallyCollapsed(true).AreaTitle(LOCTEXT("ConstantsLabel", "Constants"))
+				.BodyContent()
+				[
+					SAssignNew(SpawnScriptConstantList, SListView<TSharedPtr<EditorExposedVectorConstant>>)
+					.ItemHeight(20).ListItemsSource(&CurSpawnScript->Source->ExposedVectorConstants).OnGenerateRow(this, &SEmitterWidget::OnGenerateConstantListRow)
+				]
+			];
+	}
 }
 
 
@@ -547,5 +745,6 @@ void SNiagaraEffectEditorWidget::OnSelectionChanged(TSharedPtr<FNiagaraSimulatio
 	EffectEditor->OnEmitterSelected(SelectedItem, SelType);
 }
 
+END_SLATE_FUNCTION_BUILD_OPTIMIZATION
 
-END_SLATE_FUNCTION_BUILD_OPTIMIZATION
\ No newline at end of file
+#undef LOCTEXT_NAMESPACE
\ No newline at end of file
diff --git a/Engine/Source/Editor/NiagaraEditor/Private/SNiagaraEffectEditorWidget.h b/Engine/Source/Editor/NiagaraEditor/Private/SNiagaraEffectEditorWidget.h
--- a/Engine/Source/Editor/NiagaraEditor/Private/SNiagaraEffectEditorWidget.h
+++ b/Engine/Source/Editor/NiagaraEditor/Private/SNiagaraEffectEditorWidget.h
@@ -32,20 +32,21 @@ class SVectorConstantWidget : public SCompoundWidget, public FNotifyHook
 {
 private:	
 	FName ConstantName;
-	TSharedPtr<FNiagaraSimulation> Emitter;
+	FNiagaraEmitterScriptProperties* ScriptProps;
+
 public:
 	SLATE_BEGIN_ARGS(SVectorConstantWidget) :
-	_Emitter(nullptr),
+	_ScriptProps(nullptr),
 	_ConstantName("Undefined")
 	{
 	}
-	SLATE_ARGUMENT(TSharedPtr<FNiagaraSimulation>, Emitter)
+	SLATE_ARGUMENT(FNiagaraEmitterScriptProperties*, ScriptProps)
 	SLATE_ARGUMENT(FName, ConstantName)
 	SLATE_END_ARGS()
 
 	TOptional<float> GetConstX() const	
 	{
-		const FVector4 *ConstPtr = Emitter->GetProperties()->ExternalConstants.FindVector(ConstantName);
+		const FVector4 *ConstPtr = ScriptProps->ExternalConstants.FindVector(ConstantName);
 		if (ConstPtr)
 		{
 			return ConstPtr->X;
@@ -55,7 +56,7 @@ public:
 
 	TOptional<float> GetConstY() const
 	{
-		const FVector4 *ConstPtr = Emitter->GetProperties()->ExternalConstants.FindVector(ConstantName);
+		const FVector4 *ConstPtr = ScriptProps->ExternalConstants.FindVector(ConstantName);
 		if (ConstPtr)
 		{
 			return ConstPtr->Y;
@@ -65,7 +66,7 @@ public:
 
 	TOptional<float> GetConstZ() const
 	{
-		const FVector4 *ConstPtr = Emitter->GetProperties()->ExternalConstants.FindVector(ConstantName);
+		const FVector4 *ConstPtr = ScriptProps->ExternalConstants.FindVector(ConstantName);
 		if (ConstPtr)
 		{
 			return ConstPtr->Z;
@@ -75,7 +76,7 @@ public:
 
 	TOptional<float> GetConstW() const
 	{
-		const FVector4 *ConstPtr = Emitter->GetProperties()->ExternalConstants.FindVector(ConstantName);
+		const FVector4 *ConstPtr = ScriptProps->ExternalConstants.FindVector(ConstantName);
 		if (ConstPtr)
 		{
 			return ConstPtr->W;
@@ -91,33 +92,33 @@ public:
 	void OnConstantChangedX(float InVal, ETextCommit::Type Type)
 	{
 		FVector4 Constant(InVal, GetConstY().GetValue(), GetConstZ().GetValue(), GetConstW().GetValue());
-		Emitter->GetProperties()->ExternalConstants.SetOrAdd(ConstantName, Constant);
+		ScriptProps->ExternalConstants.SetOrAdd(ConstantName, Constant);
 	}
 
 	void OnConstantChangedY(float InVal, ETextCommit::Type Type)
 	{
 		FVector4 Constant(GetConstX().GetValue(), InVal, GetConstZ().GetValue(), GetConstW().GetValue());
-		Emitter->GetProperties()->ExternalConstants.SetOrAdd(ConstantName, Constant);
+		ScriptProps->ExternalConstants.SetOrAdd(ConstantName, Constant);
 	}
 
 	void OnConstantChangedZ(float InVal, ETextCommit::Type Type)
 	{
 		FVector4 Constant(GetConstX().GetValue(), GetConstY().GetValue(), InVal, GetConstW().GetValue());
-		Emitter->GetProperties()->ExternalConstants.SetOrAdd(ConstantName, Constant);
+		ScriptProps->ExternalConstants.SetOrAdd(ConstantName, Constant);
 	}
 
 	void OnConstantChangedW(float InVal, ETextCommit::Type Type)
 	{
 		FVector4 Constant(GetConstX().GetValue(), GetConstY().GetValue(), GetConstZ().GetValue(), InVal);
-		Emitter->GetProperties()->ExternalConstants.SetOrAdd(ConstantName, Constant);
+		ScriptProps->ExternalConstants.SetOrAdd(ConstantName, Constant);
 	}
 
 	void Construct(const FArguments& InArgs)
 	{
-		Emitter = InArgs._Emitter;
+		ScriptProps = InArgs._ScriptProps;
 		ConstantName = InArgs._ConstantName;
 
-		const FVector4 *ConstPtr = Emitter->GetProperties()->ExternalConstants.FindVector(ConstantName);
+		const FVector4 *ConstPtr = ScriptProps->ExternalConstants.FindVector(ConstantName);
 		FVector4 ConstVal;
 
 		if (ConstPtr)
@@ -172,37 +173,46 @@ class SCurveConstantWidget : public SCompoundWidget, public FNotifyHook
 {
 private:
 	FName ConstantName;
-	TSharedPtr<FNiagaraSimulation> Emitter;
+	FNiagaraEmitterScriptProperties* ScriptProps;
+	TWeakObjectPtr<UNiagaraEmitterProperties> EmitterProps;
 
 	UCurveVector *CurCurve;
 
 public:
 	SLATE_BEGIN_ARGS(SCurveConstantWidget) :
-		_Emitter(nullptr),
+		_ScriptProps(nullptr),
+		_EmitterProps(nullptr),
 		_ConstantName("Undefined")
 	{
 	}
-	SLATE_ARGUMENT(TSharedPtr<FNiagaraSimulation>, Emitter)
+	SLATE_ARGUMENT(FNiagaraEmitterScriptProperties*, ScriptProps)
+		SLATE_ARGUMENT(TWeakObjectPtr<UNiagaraEmitterProperties>, EmitterProps)
 		SLATE_ARGUMENT(FName, ConstantName)
 		SLATE_END_ARGS()
 
 	void OnCurveSelected(UObject *Asset)
 	{
+		check(ScriptProps);
 		CurCurve = Cast<UCurveVector>(Asset);
-		UNiagaraDataObject *DataObj = Emitter->GetProperties()->ExternalConstants.FindDataObj(ConstantName);
-		if (!DataObj)
+		UNiagaraEmitterProperties* PinnedProps = EmitterProps.Get();
+		if (PinnedProps)
 		{
-			DataObj = NewObject<UNiagaraCurveDataObject>(Emitter->GetParentEffect());
-			Emitter->GetProperties()->ExternalConstants.SetOrAdd(ConstantName, DataObj);
+			UNiagaraDataObject *DataObj = ScriptProps->ExternalConstants.FindDataObj(ConstantName);
+			if (!DataObj)
+			{
+				//In the details UI this is created via a combo box. Not sure which is the best way.
+				DataObj = NewObject<UNiagaraCurveDataObject>(PinnedProps);
+				ScriptProps->ExternalConstants.SetOrAdd(ConstantName, DataObj);
+			}
+
+			UNiagaraCurveDataObject *CurveData = Cast<UNiagaraCurveDataObject>(DataObj);
+			CurveData->SetCurveObject(CurCurve);
 		}
-
-		UNiagaraCurveDataObject *CurveData = Cast<UNiagaraCurveDataObject>(DataObj);
-		CurveData->SetCurveObject(CurCurve);
 	}
 
 	UObject * GetCurve() const
 	{
-		UNiagaraDataObject *DataObj = Emitter->GetProperties()->ExternalConstants.FindDataObj(ConstantName);
+		UNiagaraDataObject *DataObj = ScriptProps->ExternalConstants.FindDataObj(ConstantName);
 		UNiagaraCurveDataObject *CurveData = Cast<UNiagaraCurveDataObject>(DataObj);
 		if (CurveData)
 		{
@@ -213,7 +223,8 @@ public:
 
 	void Construct(const FArguments& InArgs)
 	{
-		Emitter = InArgs._Emitter;
+		ScriptProps = InArgs._ScriptProps;
+		EmitterProps = InArgs._EmitterProps;
 		ConstantName = InArgs._ConstantName;
 
 		this->ChildSlot
@@ -234,18 +245,67 @@ public:
 	}
 };
 
-class SEmitterWidget : public SCompoundWidget, public FNotifyHook
+class SEmitterWidgetBase : public SCompoundWidget, public FNotifyHook
 {
-private:
+protected:
 	TSharedPtr<FNiagaraSimulation> Emitter;
 	FNiagaraEffectInstance *EffectInstance;
+	class FNiagaraEffectEditor *EffectEditor;
+
+public:
+	SLATE_BEGIN_ARGS(SEmitterWidgetBase) :
+		_Emitter(nullptr),
+		_Effect(nullptr),
+		_EffectEditor(nullptr)
+	{
+	}
+	SLATE_ARGUMENT(TSharedPtr<FNiagaraSimulation>, Emitter)
+	SLATE_ARGUMENT(FNiagaraEffectInstance*, Effect)
+	SLATE_ARGUMENT(FNiagaraEffectEditor*, EffectEditor)
+	SLATE_END_ARGS()
+
+
+	void OnEmitterNameChanged(const FText &InText)
+	{
+		Emitter->GetProperties()->EmitterName = InText.ToString();
+	}
+
+	FText GetEmitterName() const
+	{
+		UNiagaraEmitterProperties* PinnedProps = Emitter.IsValid() ? Emitter->GetProperties().Get() : NULL;
+		return PinnedProps ? FText::FromString(PinnedProps->EmitterName) : FText();
+	}
+
+	void OnEmitterEnabledChanged(ECheckBoxState NewCheckedState)
+	{
+		const bool bNewEnabledState = (NewCheckedState == ESlateCheckBoxState::Checked);
+		Emitter->SetEnabled(bNewEnabledState);
+	}
+
+	ECheckBoxState IsEmitterEnabled() const
+	{
+		return Emitter->IsEnabled() ? ECheckBoxState::Checked : ECheckBoxState::Unchecked;
+	}
+
+	FText GetStatsText() const
+	{
+		TCHAR Txt[128];
+		FCString::Snprintf(Txt, 128, TEXT("%i Particles  |  %.2f ms  |  %.1f MB"), Emitter->GetNumParticles(), Emitter->GetTotalCPUTime(), Emitter->GetTotalBytesUsed()/(1024.0*1024.0));
+		return FText::FromString(Txt);
+	}
+
+	void Construct(const FArguments& InArgs);
+};
+
+class SEmitterWidget : public SEmitterWidgetBase
+{
+protected:
 	TSharedPtr<STextBlock> UpdateComboText, SpawnComboText;
 	TSharedPtr<SComboButton> UpdateCombo, SpawnCombo;
 	TSharedPtr<SContentReference> UpdateScriptSelector, SpawnScriptSelector;
 	TSharedPtr<SWidget> ThumbnailWidget;
 	TSharedPtr<FAssetThumbnailPool> ThumbnailPool;
 	TSharedPtr<FAssetThumbnail> MaterialThumbnail;
-	UNiagaraScript *CurUpdateScript, *CurSpawnScript;
 	UMaterial *CurMaterial;
 
 	TSharedPtr<IDetailsView> RendererPropertiesView;
@@ -258,20 +318,14 @@ private:
 	SVerticalBox::FSlot *SpawnScriptConstantListSlot;
 
 public:
-	SLATE_BEGIN_ARGS(SEmitterWidget) :
-		_Emitter(nullptr)
-	{
-	}
-	SLATE_ARGUMENT(TSharedPtr<FNiagaraSimulation>, Emitter)
-	SLATE_ARGUMENT(FNiagaraEffectInstance*, Effect)
-	SLATE_END_ARGS()
+
 
 	TSharedRef<ITableRow> OnGenerateConstantListRow(TSharedPtr<EditorExposedVectorConstant> InItem, const TSharedRef< STableViewBase >& OwnerTable)
 	{
 		return SNew(STableRow<TSharedPtr<FString>>, OwnerTable)
 		.Content()
 		[
-			SNew(SVectorConstantWidget).ConstantName(InItem->ConstName).Emitter(Emitter)
+			SNew(SVectorConstantWidget).ConstantName(InItem->ConstName).ScriptProps(&Emitter->GetProperties()->UpdateScriptProps)
 		];
 	}
 
@@ -281,19 +335,19 @@ public:
 		return SNew(STableRow<TSharedPtr<FString>>, OwnerTable)
 			.Content()
 			[
-				SNew(SCurveConstantWidget).ConstantName(InItem->ConstName).Emitter(Emitter)
+				SNew(SCurveConstantWidget).ConstantName(InItem->ConstName).ScriptProps(&Emitter->GetProperties()->UpdateScriptProps).EmitterProps(Emitter->GetProperties())
 			];
 	}
 
 
 	void OnEmitterNameChanged(const FText &InText)
 	{
-		Emitter->GetProperties()->Name = InText.ToString();
+		Emitter->GetProperties()->EmitterName = InText.ToString();
 	}
 
 	FText GetEmitterName() const
 	{
-		return FText::FromString(Emitter->GetProperties()->Name);
+		return FText::FromString(Emitter->GetProperties()->EmitterName);
 	}
 
 	void OnUpdateScriptSelectedFromPicker(UObject *Asset);
@@ -318,12 +372,12 @@ public:
 
 	UObject *GetUpdateScript() const
 	{
-		return Emitter->GetProperties()->UpdateScript;
+		return Emitter->GetProperties()->UpdateScriptProps.Script;
 	}
 
 	UObject *GetSpawnScript() const
 	{
-		return Emitter->GetProperties()->SpawnScript;
+		return Emitter->GetProperties()->SpawnScriptProps.Script;
 	}
 
 	void OnEmitterEnabledChanged(ECheckBoxState NewCheckedState)
@@ -349,7 +403,7 @@ public:
 	FText GetSpawnRateText() const
 	{
 		TCHAR Txt[32];
-		FCString::Snprintf(Txt, 32, TEXT("%f"), Emitter->GetProperties()->SpawnRate );
+		FCString::Snprintf(Txt, 32, TEXT("%f"), Emitter->GetProperties()->SpawnRate);
 		return FText::FromString(Txt);
 	}
 
@@ -374,27 +428,31 @@ public:
 			return NSLOCTEXT("NiagaraEffectEditor", "None", "None");
 		}
 
-		return FText::FromString(*RenderModuleList[Emitter->GetProperties()->RenderModuleType-1]);
+		return FText::FromString(*RenderModuleList[Emitter->GetProperties()->RenderModuleType - 1]);
 	}
 
-	TSharedRef<SWidget>GenerateRenderModuleComboWidget(TSharedPtr<FString> WidgetString )
+	TSharedRef<SWidget>GenerateRenderModuleComboWidget(TSharedPtr<FString> WidgetString)
 	{
 		return SNew(STextBlock)
 			.Text(FText::FromString(*WidgetString));
 	}
 
-	FText GetStatsText() const
-	{
-		TCHAR Txt[128];
-		FCString::Snprintf(Txt, 128, TEXT("%i Particles  |  %.2f ms  |  %.1f MB"), Emitter->GetNumParticles(), Emitter->GetTotalCPUTime(), Emitter->GetTotalBytesUsed()/(1024.0*1024.0));
-		return FText::FromString(Txt);
-	}
-
 	void Construct(const FArguments& InArgs);
 
+};
 
+class SEmitterWidgetDev : public SEmitterWidgetBase
+{
+protected:
+	/** Details view used for the dev view. */
+	TSharedPtr<IDetailsView> Details;
 
+	void BuildContents();
+public:
+
+	void Construct(const FArguments& InArgs);
 
+	virtual void NotifyPostChange(const FPropertyChangedEvent& PropertyChangedEvent, UProperty* PropertyThatChanged)override;
 };
 
 
@@ -442,15 +500,23 @@ private:
 	TSharedPtr< SListView<TSharedPtr<FNiagaraSimulation> > > ListView;
 
 	class FNiagaraEffectEditor *EffectEditor;
+
+	/**Scroll Box for dev emitters view.*/
+	TSharedPtr<SScrollBox> EmitterBox;
+
+	/** True if this is the development emitters view.*/
+	bool bForDev;
 public:
 	SLATE_BEGIN_ARGS(SNiagaraEffectEditorWidget)
 		: _EffectObj(nullptr)
+		, _bForDev(false)
 	{
 	}
 	SLATE_ARGUMENT(UNiagaraEffect*, EffectObj)
 	SLATE_ARGUMENT(TSharedPtr<SWidget>, TitleBar)
 	SLATE_ARGUMENT(FNiagaraEffectInstance*, EffectInstance)
 	SLATE_ARGUMENT(class FNiagaraEffectEditor*, EffectEditor)
+	SLATE_ARGUMENT(bool, bForDev)
 	SLATE_END_ARGS()
 
 
@@ -462,7 +528,10 @@ public:
 
 	void UpdateList()
 	{
-		ListView->RequestListRefresh();
+		if (bForDev)
+			ConstructDevEmittersView();
+		else
+			ListView->RequestListRefresh();
 	}
 
 	TSharedRef<ITableRow> OnGenerateWidgetForList(TSharedPtr<FNiagaraSimulation> InItem, const TSharedRef< STableViewBase >& OwnerTable)
@@ -474,7 +543,7 @@ public:
 				//.BorderImage(FEditorStyle::GetBrush("ToolPanel.GroupBorder"))
 				.Padding(2.0f)
 				[
-					SNew(SEmitterWidget).Emitter(InItem).Effect(EffectInstance)
+					SNew(SEmitterWidget).Emitter(InItem).Effect(EffectInstance).EffectEditor(EffectEditor)
 				]
 			];
 	}
@@ -483,5 +552,6 @@ public:
 
 	void Construct(const FArguments& InArgs);
 
-};
+	void ConstructDevEmittersView();
 
+};
diff --git a/Engine/Source/Editor/UnrealEd/Classes/Niagara/NiagaraNodeInput.h b/Engine/Source/Editor/UnrealEd/Classes/Niagara/NiagaraNodeInput.h
--- a/Engine/Source/Editor/UnrealEd/Classes/Niagara/NiagaraNodeInput.h
+++ b/Engine/Source/Editor/UnrealEd/Classes/Niagara/NiagaraNodeInput.h
@@ -26,6 +26,9 @@ public:
 	UPROPERTY(EditAnywhere, Category = Constant)
 	FMatrix MatrixDefault;
 
+	UPROPERTY(Instanced, EditAnywhere, Category = Constant)
+	UNiagaraDataObject* DataObjectDefault;
+
 	/** Allows code to explicitly disable exposing of certain inputs e.g. system constants such as Delta Time. */
 	UPROPERTY()
 	uint32 bCanBeExposed:1;
diff --git a/Engine/Source/Editor/UnrealEd/Private/Niagara/NiagaraNodeFunctionCall.cpp b/Engine/Source/Editor/UnrealEd/Private/Niagara/NiagaraNodeFunctionCall.cpp
--- a/Engine/Source/Editor/UnrealEd/Private/Niagara/NiagaraNodeFunctionCall.cpp
+++ b/Engine/Source/Editor/UnrealEd/Private/Niagara/NiagaraNodeFunctionCall.cpp
@@ -22,6 +22,8 @@ void UNiagaraNodeFunctionCall::PostEditChangeProperty(struct FPropertyChangedEve
 		ReallocatePins();
 	}
 	Super::PostEditChangeProperty(PropertyChangedEvent);
+
+	GetGraph()->NotifyGraphChanged();
 }
 
 void UNiagaraNodeFunctionCall::ReallocatePins()
diff --git a/Engine/Source/Editor/UnrealEd/Private/Niagara/NiagaraNodeInput.cpp b/Engine/Source/Editor/UnrealEd/Private/Niagara/NiagaraNodeInput.cpp
--- a/Engine/Source/Editor/UnrealEd/Private/Niagara/NiagaraNodeInput.cpp
+++ b/Engine/Source/Editor/UnrealEd/Private/Niagara/NiagaraNodeInput.cpp
@@ -12,6 +12,7 @@ UNiagaraNodeInput::UNiagaraNodeInput(const FObjectInitializer& ObjectInitializer
 , MatrixDefault(FMatrix::Identity)
 , bCanBeExposed(true)
 , bExposeWhenConstant(true)
+, DataObjectDefault(nullptr)
 {
 }
 
@@ -122,6 +123,8 @@ void UNiagaraNodeInput::ReallocatePins()
 #endif
 	}
 	OldPins.Empty();
+
+	GetGraph()->NotifyGraphChanged();
 }
 
 void UNiagaraNodeInput::AllocateDefaultPins()
@@ -186,18 +189,18 @@ void UNiagaraNodeInput::Compile(class INiagaraCompiler* Compiler, TArray<FNiagar
 			case ENiagaraDataType::Scalar:	Outputs.Add(FNiagaraNodeResult(Compiler->GetExternalConstant(Input, FloatDefault), Pins[0]));	break;
 			case ENiagaraDataType::Vector:	Outputs.Add(FNiagaraNodeResult(Compiler->GetExternalConstant(Input, VectorDefault), Pins[0]));	break;
 			case ENiagaraDataType::Matrix:	Outputs.Add(FNiagaraNodeResult(Compiler->GetExternalConstant(Input, MatrixDefault), Pins[0]));	break;
-			case ENiagaraDataType::Curve:	Outputs.Add(FNiagaraNodeResult(Compiler->GetExternalCurveConstant(Input), Pins[0]));	break;
+			case ENiagaraDataType::Curve:	Outputs.Add(FNiagaraNodeResult(Compiler->GetExternalConstant(Input, DataObjectDefault), Pins[0]));	break;
 			}
 		}
 		else
 		{
-			//treat this input as an external constant.
+			//treat this input as an internal constant.
 			switch (Input.Type)
 			{
 			case ENiagaraDataType::Scalar:	Outputs.Add(FNiagaraNodeResult(Compiler->GetInternalConstant(Input, FloatDefault), Pins[0]));	break;
 			case ENiagaraDataType::Vector:	Outputs.Add(FNiagaraNodeResult(Compiler->GetInternalConstant(Input, VectorDefault), Pins[0]));	break;
 			case ENiagaraDataType::Matrix:	Outputs.Add(FNiagaraNodeResult(Compiler->GetInternalConstant(Input, MatrixDefault), Pins[0]));	break;
-			case ENiagaraDataType::Curve:	Outputs.Add(FNiagaraNodeResult(Compiler->GetExternalCurveConstant(Input), Pins[0]));	break;//Internal curve Constant?
+			case ENiagaraDataType::Curve:	Outputs.Add(FNiagaraNodeResult(Compiler->GetInternalConstant(Input, DataObjectDefault), Pins[0]));	break;//Internal curve Constant?
 			}
 		}
 	}
diff --git a/Engine/Source/Editor/UnrealEd/Private/Niagara/NiagaraScriptSource.cpp b/Engine/Source/Editor/UnrealEd/Private/Niagara/NiagaraScriptSource.cpp
--- a/Engine/Source/Editor/UnrealEd/Private/Niagara/NiagaraScriptSource.cpp
+++ b/Engine/Source/Editor/UnrealEd/Private/Niagara/NiagaraScriptSource.cpp
@@ -106,8 +106,10 @@ void UNiagaraScriptSource::Compile()
 	FNiagaraEditorModule& NiagaraEditorModule = FModuleManager::Get().LoadModuleChecked<FNiagaraEditorModule>(TEXT("NiagaraEditor"));
 	NiagaraEditorModule.CompileScript(ScriptOwner);
 
-	FNiagaraConstants& ExternalConsts = ScriptOwner->ConstantData.ExternalConstants;
+	FNiagaraConstants& ExternalConsts = ScriptOwner->ConstantData.ExposedConstants;
 
+	//Build the constant list. 
+	//This is mainly just jumping through some hoops for the custom UI. Should be removed and have the UI just read directly from the constants stored in the UScript.
 	const UEdGraphSchema_Niagara* Schema = CastChecked<UEdGraphSchema_Niagara>(NodeGraph->GetSchema());
 	ExposedVectorConstants.Empty();
 	for (int32 ConstIdx = 0; ConstIdx < ExternalConsts.GetNumVectorConstants(); ConstIdx++)
@@ -130,9 +132,13 @@ void UNiagaraScriptSource::Compile()
 	{
 		EditorExposedVectorCurveConstant *Const = new EditorExposedVectorCurveConstant();
 		FNiagaraVariableInfo Info;
-		UNiagaraDataObject* Dummy;
-		ExternalConsts.GetDataObjectConstant(ConstIdx, Dummy, Info);
+		UNiagaraDataObject* Obj;
+		ExternalConsts.GetDataObjectConstant(ConstIdx, Obj, Info);
 		Const->ConstName = Info.Name;
+		if (UNiagaraCurveDataObject* CurvObj = Cast<UNiagaraCurveDataObject>(Obj))
+		{
+			Const->Value = CurvObj->CurveObj;
+		}
 		//Set default value for this too?
 		ExposedVectorCurveConstants.Add(MakeShareable(Const));
 	}
diff --git a/Engine/Source/Editor/UnrealEd/Public/INiagaraCompiler.h b/Engine/Source/Editor/UnrealEd/Public/INiagaraCompiler.h
--- a/Engine/Source/Editor/UnrealEd/Public/INiagaraCompiler.h
+++ b/Engine/Source/Editor/UnrealEd/Public/INiagaraCompiler.h
@@ -5,6 +5,7 @@
 #include "NiagaraEditorCommon.h"
 
 class UNiagaraNode;
+class UNiagaraDataObject;
 
 //Interface for Niagara compilers.
 class UNREALED_API INiagaraCompiler
@@ -25,6 +26,9 @@ public:
 	virtual TNiagaraExprPtr GetExternalConstant(const FNiagaraVariableInfo& Constant, const FVector4& Default) = 0;
 	/** Creates an expression which fetches the named constant. */
 	virtual TNiagaraExprPtr GetExternalConstant(const FNiagaraVariableInfo& Constant, const FMatrix& Default) = 0;
+	/** Creates an expression which fetches the named CURVE. */
+	virtual TNiagaraExprPtr GetExternalConstant(const FNiagaraVariableInfo& Constant, const UNiagaraDataObject* Default) = 0;
+
 	/** Creates an expression which fetches the named constant. */
 	virtual TNiagaraExprPtr GetInternalConstant(const FNiagaraVariableInfo& Constant, float Default) = 0;
 	/** Creates an expression which fetches the named constant. */
@@ -32,7 +36,7 @@ public:
 	/** Creates an expression which fetches the named constant. */
 	virtual TNiagaraExprPtr GetInternalConstant(const FNiagaraVariableInfo& Constant, const FMatrix& Default) = 0;
 	/** Creates an expression which fetches the named CURVE. */
-	virtual TNiagaraExprPtr GetExternalCurveConstant(const FNiagaraVariableInfo& Constant) = 0;
+	virtual TNiagaraExprPtr GetInternalConstant(const FNiagaraVariableInfo& Constant, const UNiagaraDataObject* Default) = 0;
 
 	/** Creates an expression which writes the result of SourceExpression to the named attribute. */
 	virtual TNiagaraExprPtr Output(const FNiagaraVariableInfo& Attr, TNiagaraExprPtr& SourceExpression) = 0;
diff --git a/Engine/Source/Runtime/Engine/Classes/Engine/NiagaraEffectRendererProperties.h b/Engine/Source/Runtime/Engine/Classes/Engine/NiagaraEffectRendererProperties.h
--- a/Engine/Source/Runtime/Engine/Classes/Engine/NiagaraEffectRendererProperties.h
+++ b/Engine/Source/Runtime/Engine/Classes/Engine/NiagaraEffectRendererProperties.h
@@ -6,7 +6,7 @@
 /**
 * Emitter properties base class
 * Each Effectrenderer derives from this with its own class, and returns it in GetProperties; a copy
-* of those specific properties is stored on FNiagaraEmitterProperties (on the effect) for serialization 
+* of those specific properties is stored on UNiagaraEmitterProperties (on the effect) for serialization 
 * and handed back to the effect renderer on load.
 */
 UCLASS()
diff --git a/Engine/Source/Runtime/Niagara/Classes/NiagaraEffect.h b/Engine/Source/Runtime/Niagara/Classes/NiagaraEffect.h
--- a/Engine/Source/Runtime/Niagara/Classes/NiagaraEffect.h
+++ b/Engine/Source/Runtime/Niagara/Classes/NiagaraEffect.h
@@ -22,14 +22,15 @@ class NIAGARA_API UNiagaraEffect : public UObject
 	GENERATED_UCLASS_BODY()
 
 public:
-	FNiagaraEmitterProperties* AddEmitterProperties();
+	UNiagaraEmitterProperties* AddEmitterProperties(UNiagaraEmitterProperties *Props = nullptr);
+	void DeleteEmitterProperties(UNiagaraEmitterProperties *Props);
 
 	void Init()
 	{
 	}
 
 
-	FNiagaraEmitterProperties *GetEmitterProperties(int Idx)
+	UNiagaraEmitterProperties *GetEmitterProperties(int Idx)
 	{
 		check(Idx < EmitterProps.Num());
 		return EmitterProps[Idx];
@@ -42,15 +43,10 @@ public:
 
 	void CreateEffectRendererProps(TSharedPtr<FNiagaraSimulation> Sim);
 
-	virtual void PostLoad() override; 
-	virtual void PreSave() override;
-
 private:
-	// serialized array of emitter properties
-	UPROPERTY()
-	TArray<FNiagaraEmitterProperties>EmitterPropsSerialized;
 
-	TArray<FNiagaraEmitterProperties*> EmitterProps;
+	UPROPERTY()
+	TArray<UNiagaraEmitterProperties*> EmitterProps;
 };
 
 
@@ -106,7 +102,8 @@ public:
 		Age = 0.0f;
 	}
 
-	NIAGARA_API TSharedPtr<FNiagaraSimulation> AddEmitter(FNiagaraEmitterProperties *Properties);
+	NIAGARA_API TSharedPtr<FNiagaraSimulation> AddEmitter(UNiagaraEmitterProperties *Properties);
+	NIAGARA_API void DeleteEmitter(TSharedPtr<FNiagaraSimulation> Emitter);
 
 	void SetConstant(FNiagaraVariableInfo ID, const float Value)
 	{
@@ -128,58 +125,7 @@ public:
 		Constants.SetOrAdd(ConstantName, Value);
 	}
 
-
-
-	void Tick(float DeltaSeconds)
-	{
-		Constants.SetOrAdd(FName(TEXT("EffectGrid")), VolumeGrid); 
-
-		// pass the constants down to the emitter
-		// TODO: should probably just pass a pointer to the table
-		EffectBounds.Init();
-
-		for (TSharedPtr<FNiagaraSimulation>&it : Emitters)
-		{
-			FNiagaraEmitterProperties *Props = it->GetProperties();
-
-			int Duration = Props->EndTime - Props->StartTime;
-			int LoopedStartTime = Props->StartTime + Duration*it->GetLoopCount();
-			int LoopedEndTime = Props->EndTime + Duration*it->GetLoopCount();
-
-			// manage emitter lifetime
-			//
-			if (	(Props->StartTime == 0.0f && Props->EndTime == 0.0f)
-					|| (LoopedStartTime<Age && LoopedEndTime>Age)
-				)
-			{
-				it->SetTickState(NTS_Running);
-			}
-			else
-			{
-				// if we're past end time, manage looping; we reset the emitters age constant
-				// if it has one
-				if (Props->NumLoops > 1 && it->GetLoopCount()<Props->NumLoops)
-				{
-					it->LoopRestart();
-				}
-				else
-				{
-					it->SetTickState(NTS_Dieing);
-				}
-			}
-
-			if (it->GetTickState() != NTS_Dead && it->GetTickState() != NTS_Suspended)
-			{
-				it->SetConstants(Constants);
-				it->GetConstants().Merge(it->GetProperties()->ExternalConstants);
-				it->Tick(DeltaSeconds);
-			}
-
-			EffectBounds += it->GetEffectRenderer()->GetBounds();
-		}
-
-		Age += DeltaSeconds;
-	}
+	void Tick(float DeltaSeconds);
 
 	NIAGARA_API void RenderModuleupdate();
 
diff --git a/Engine/Source/Runtime/Niagara/Classes/NiagaraSimulation.h b/Engine/Source/Runtime/Niagara/Classes/NiagaraSimulation.h
--- a/Engine/Source/Runtime/Niagara/Classes/NiagaraSimulation.h
+++ b/Engine/Source/Runtime/Niagara/Classes/NiagaraSimulation.h
@@ -28,79 +28,93 @@ enum EEmitterRenderModuleType
 };
 
 
+USTRUCT()
+struct FNiagaraEmitterScriptProperties
+{
+	GENERATED_BODY()
+
+	UPROPERTY(EditAnywhere, Category = "Script")
+	UNiagaraScript *Script;
+
+	UPROPERTY(EditAnywhere, Category = "Script", meta = (ShowOnlyInnerProperties))
+	FNiagaraConstants ExternalConstants;
+
+	void Init(UNiagaraEmitterProperties* EmitterProps)
+	{
+		if (Script)
+		{
+			ExternalConstants.Init(EmitterProps, this);
+		}
+		else
+		{
+			ExternalConstants.Empty();
+		}
+	}
+};
 
 
 /** 
  *	FniagaraEmitterProperties stores the attributes of an FNiagaraSimulation
  *	that need to be serialized and are used for its initialization 
  */
-USTRUCT()
-struct FNiagaraEmitterProperties
+UCLASS(MinimalAPI)
+class UNiagaraEmitterProperties : public UObject
 {
-	GENERATED_USTRUCT_BODY()
+	GENERATED_UCLASS_BODY()
 public:
-	FNiagaraEmitterProperties() : 
-		bIsEnabled(true),
-		SpawnRate(50),
-		UpdateScript(nullptr),
-		SpawnScript(nullptr),
-		Material(nullptr), 
-		RenderModuleType(RMT_Sprites),
-		StartTime(0.0f),
-		EndTime(0.0f),
-		RendererProperties(nullptr),
-		NumLoops(0)
+	void Init()
 	{
-		Name = FString(TEXT("New Emitter"));
+		SpawnScriptProps.Init(this);
+		UpdateScriptProps.Init(this);
 	}
-	UPROPERTY(EditAnywhere, Category = "Emitter Properties")
-	FString Name;
-	UPROPERTY(EditAnywhere, Category = "Emitter Properties")
+
+	//Begin UObject Interface
+	void PostEditChangeProperty(struct FPropertyChangedEvent& PropertyChangedEvent);
+	void PostLoad();
+	void Serialize(FArchive& Ar);
+	//End UObject Interface
+
+	UPROPERTY(EditAnywhere, Category = "Emitter")
+	FString EmitterName;
+	UPROPERTY(EditAnywhere, Category = "Emitter")
 	bool bIsEnabled;
-	UPROPERTY(EditAnywhere, Category = "Emitter Properties")
+	UPROPERTY(EditAnywhere, Category = "Emitter")
 	float SpawnRate;
-	UPROPERTY(EditAnywhere, Category = "Emitter Properties")
-	UNiagaraScript *UpdateScript;
-	UPROPERTY(EditAnywhere, Category = "Emitter Properties")
-	UNiagaraScript *SpawnScript;
-	UPROPERTY(EditAnywhere, Category = "Emitter Properties")
+	UPROPERTY(EditAnywhere, Category = "Emitter")
 	UMaterial *Material;
-	UPROPERTY(EditAnywhere, Category = "Emitter Properties")
+	UPROPERTY(EditAnywhere, Category = "Emitter")
 	TEnumAsByte<EEmitterRenderModuleType> RenderModuleType;
 
-	UPROPERTY(EditAnywhere, Category = "Emitter Properties")
+	UPROPERTY(EditAnywhere, Category = "Emitter")
 	float StartTime;
-	UPROPERTY(EditAnywhere, Category = "Emitter Properties")
+	UPROPERTY(EditAnywhere, Category = "Emitter")
 	float EndTime;
 
-
 	UPROPERTY()
 	class UNiagaraEffectRendererProperties *RendererProperties;
 
-	UPROPERTY(EditAnywhere, Category = "Emitter Properties")
-	FNiagaraConstantMap ExternalConstants;		// these are the update script constants from the effect editor; will be added to the emitter's constant map
-
-	UPROPERTY(EditAnywhere, Category = "Emitter Properties")
-	FNiagaraConstantMap ExternalSpawnConstants;		// these are the spawn script constants from the effect editor; will be added to the emitter's constant map
-
-	UPROPERTY(EditAnywhere, Category = "Emitter Properties")
+	UPROPERTY(EditAnywhere, Category = "Emitter")
 	int32 NumLoops;
-};
 
+	UPROPERTY(EditAnywhere, Category = "Emitter", meta = (ShowOnlyInnerProperties))
+	FNiagaraEmitterScriptProperties UpdateScriptProps;
 
+	UPROPERTY(EditAnywhere, Category = "Emitter", meta = (ShowOnlyInnerProperties))
+	FNiagaraEmitterScriptProperties SpawnScriptProps;
+};
 
 /**
-* A niagara particle simulation.
+* A Niagara particle simulation.
 */
 struct FNiagaraSimulation
 {
 public:
-	explicit FNiagaraSimulation(FNiagaraEmitterProperties *InProps, UNiagaraEffect *Effect);
-	FNiagaraSimulation(FNiagaraEmitterProperties *Props, UNiagaraEffect *Effect, ERHIFeatureLevel::Type InFeatureLevel);
+	explicit FNiagaraSimulation(TWeakObjectPtr<UNiagaraEmitterProperties> InProps, UNiagaraEffect* Effect);
+	FNiagaraSimulation(TWeakObjectPtr<UNiagaraEmitterProperties> Props, UNiagaraEffect* Effect, ERHIFeatureLevel::Type InFeatureLevel);
 	virtual ~FNiagaraSimulation()
 	{}
 
-	void Init();
+	NIAGARA_API void Init();
 
 	void Tick(float DeltaSeconds);
 
@@ -121,16 +135,14 @@ public:
 
 	int32 GetNumParticles()	{ return Data.GetNumParticles(); }
 
-	FNiagaraEmitterProperties *GetProperties()	{ return Props; }
-	void SetProperties(FNiagaraEmitterProperties *InProps);
+	TWeakObjectPtr<UNiagaraEmitterProperties> GetProperties()	{ return Props; }
+	void SetProperties(TWeakObjectPtr<UNiagaraEmitterProperties> InProps);
 
 	UNiagaraEffect *GetParentEffect()	{ return ParentEffect; }
 
 	float NIAGARA_API GetTotalCPUTime();
 	int	NIAGARA_API GetTotalBytesUsed();
 
-	void NIAGARA_API SetSpawnScript(UNiagaraScript* Script);
-	void NIAGARA_API SetUpdateScript(UNiagaraScript* Script);
 	void NIAGARA_API SetEnabled(bool bEnabled);
 
 	float GetAge() { return Age; }
@@ -145,7 +157,7 @@ public:
 	ENiagaraTickState NIAGARA_API GetTickState()	{ return TickState; }
 	void NIAGARA_API SetTickState(ENiagaraTickState InState)	{ TickState = InState; }
 private:
-	FNiagaraEmitterProperties *Props;		// points to an entry in the array of FNiagaraProperties stored in the EffectInstance (itself pointing to the effect's properties)
+	TWeakObjectPtr<UNiagaraEmitterProperties> Props;		// points to an entry in the array of FNiagaraProperties stored in the EffectInstance (itself pointing to the effect's properties)
 
 	/* The age of the emitter*/
 	float Age;
diff --git a/Engine/Source/Runtime/Niagara/NiagaraConstantSet.h b/Engine/Source/Runtime/Niagara/NiagaraConstantSet.h
--- a/Engine/Source/Runtime/Niagara/NiagaraConstantSet.h
+++ b/Engine/Source/Runtime/Niagara/NiagaraConstantSet.h
@@ -6,331 +6,217 @@
 #include "VectorVMDataObject.h"
 #include "NiagaraConstantSet.generated.h"
 
-USTRUCT()
+struct FNiagaraConstants;
+
+/** Runtime storage for emitter overridden constants. */
 struct FNiagaraConstantMap
 {
-	GENERATED_USTRUCT_BODY()
-
 private:
 	TMap<FNiagaraVariableInfo, float> ScalarConstants;
 	TMap<FNiagaraVariableInfo, FVector4> VectorConstants;
 	TMap<FNiagaraVariableInfo, FMatrix> MatrixConstants;
-	TMap<FNiagaraVariableInfo, UNiagaraDataObject*> DataConstants;
+	TMap<FNiagaraVariableInfo, class UNiagaraDataObject*> DataConstants;
 
 public:
 
-	void Empty()
-	{
-		ScalarConstants.Empty();
-		VectorConstants.Empty();
-		MatrixConstants.Empty();
-	}
+	void Empty();
+
+	void SetOrAdd(FNiagaraVariableInfo ID, float Sc);
+	void SetOrAdd(FNiagaraVariableInfo ID, const FVector4& Vc);
+	void SetOrAdd(FNiagaraVariableInfo ID, const FMatrix& Mc);
+	void SetOrAdd(FNiagaraVariableInfo ID, UNiagaraDataObject* Cc);
+	void SetOrAdd(FName Name, float Sc);
+	void SetOrAdd(FName Name, const FVector4& Vc);
+	void SetOrAdd(FName Name, const FMatrix& Mc);
+	void SetOrAdd(FName Name, class UNiagaraDataObject* Cc);
+
+	float* FindScalar(FNiagaraVariableInfo ID);
+	FVector4* FindVector(FNiagaraVariableInfo ID);
+	FMatrix* FindMatrix(FNiagaraVariableInfo ID);
+	UNiagaraDataObject* FindDataObj(FNiagaraVariableInfo ID);
+	const float* FindScalar(FNiagaraVariableInfo ID)const;
+	const FVector4* FindVector(FNiagaraVariableInfo ID)const;
+	const FMatrix* FindMatrix(FNiagaraVariableInfo ID)const;
+	UNiagaraDataObject* FindDataObj(FNiagaraVariableInfo ID) const;
+	float* FindScalar(FName Name);
+	FVector4* FindVector(FName Name);
+	FMatrix* FindMatrix(FName Name);
+	UNiagaraDataObject* FindDataObj(FName Name);
+	const float* FindScalar(FName Name)const;
+	const FVector4* FindVector(FName Name)const;
+	const FMatrix* FindMatrix(FName Name)const;
+	UNiagaraDataObject* FindDataObj(FName Name)const;
+
+	void Merge(FNiagaraConstants &InConstants);
+	void Merge(FNiagaraConstantMap &InMap);
+
+	const TMap<FNiagaraVariableInfo, class UNiagaraDataObject*> &GetDataConstants()	const;
+};
 
-	void SetOrAdd(FNiagaraVariableInfo ID, float Sc)
-	{
-		check(ID.Type == ENiagaraDataType::Scalar);
-		ScalarConstants.FindOrAdd(ID) = Sc;
-	}
+USTRUCT()
+struct FNiagaraConstantBase
+{
+	GENERATED_USTRUCT_BODY()
 
-	void SetOrAdd(FNiagaraVariableInfo ID, const FVector4& Vc)
-	{
-		check(ID.Type == ENiagaraDataType::Vector);
-		VectorConstants.FindOrAdd(ID) = Vc;
-	}
+	FNiagaraConstantBase()
+	: Name(NAME_None)
+	{}
 
-	void SetOrAdd(FNiagaraVariableInfo ID, const FMatrix& Mc)
-	{
-		check(ID.Type == ENiagaraDataType::Matrix);
-		MatrixConstants.FindOrAdd(ID) = Mc;
-	}
+	FNiagaraConstantBase(FName InName)
+	: Name(InName)
+	{}
 
-	void SetOrAdd(FNiagaraVariableInfo ID, UNiagaraDataObject* Cc)
-	{
-		check(ID.Type == ENiagaraDataType::Curve);
-		DataConstants.FindOrAdd(ID) = Cc;
-	}
+	UPROPERTY(VisibleAnywhere, Category = "Constant")
+	FName Name;
+};
 
-	void SetOrAdd(FName Name, float Sc)
-	{
-		ScalarConstants.FindOrAdd(FNiagaraVariableInfo(Name, ENiagaraDataType::Scalar)) = Sc;
-	}
+USTRUCT()
+struct FNiagaraConstants_Float : public FNiagaraConstantBase
+{
+	GENERATED_USTRUCT_BODY()
 
-	void SetOrAdd(FName Name, const FVector4& Vc)
-	{
-		VectorConstants.FindOrAdd(FNiagaraVariableInfo(Name, ENiagaraDataType::Vector)) = Vc;
-	}
+	FNiagaraConstants_Float()
+	: Value(0.0f)
+	{}
 
-	void SetOrAdd(FName Name, const FMatrix& Mc)
-	{
-		MatrixConstants.FindOrAdd(FNiagaraVariableInfo(Name, ENiagaraDataType::Matrix)) = Mc;
-	}
+	FNiagaraConstants_Float(FName InName, float InValue)
+	: FNiagaraConstantBase(InName)
+	, Value(InValue)
+	{}
 
-	void SetOrAdd(FName Name, UNiagaraDataObject* Cc)
-	{
-		DataConstants.FindOrAdd(FNiagaraVariableInfo(Name, ENiagaraDataType::Curve)) = Cc;
-	}
+	UPROPERTY(EditAnywhere, Category = "Constant")
+	float Value;
+};
 
+USTRUCT()
+struct FNiagaraConstants_Vector : public FNiagaraConstantBase
+{
+	GENERATED_USTRUCT_BODY()
 
-	float* FindScalar(FNiagaraVariableInfo ID){ return ScalarConstants.Find(ID); }
-	FVector4* FindVector(FNiagaraVariableInfo ID){ return VectorConstants.Find(ID); }
-	FMatrix* FindMatrix(FNiagaraVariableInfo ID){ return MatrixConstants.Find(ID); }
-	const float* FindScalar(FNiagaraVariableInfo ID)const { return ScalarConstants.Find(ID); }
-	const FVector4* FindVector(FNiagaraVariableInfo ID)const { return VectorConstants.Find(ID); }
-	const FMatrix* FindMatrix(FNiagaraVariableInfo ID)const { return MatrixConstants.Find(ID); }
-	UNiagaraDataObject* FindDataObj(FNiagaraVariableInfo ID) const
-	{
-		UNiagaraDataObject * const *Ptr = DataConstants.Find(ID);
-		if (Ptr)
-		{
-			return *Ptr;
-		}
+	FNiagaraConstants_Vector()
+	: Value(FVector4(1.0f, 1.0f, 1.0f, 1.0f))
+	{}
 
-		return nullptr; 
-	}
-	float* FindScalar(FName Name) { return ScalarConstants.Find(FNiagaraVariableInfo(Name,ENiagaraDataType::Scalar)); }
-	FVector4* FindVector(FName Name) { return VectorConstants.Find(FNiagaraVariableInfo(Name, ENiagaraDataType::Vector)); }
-	FMatrix* FindMatrix(FName Name) { return MatrixConstants.Find(FNiagaraVariableInfo(Name, ENiagaraDataType::Matrix)); }
-	UNiagaraDataObject* FindDataObj(FName Name)const { return FindDataObj(FNiagaraVariableInfo(Name, ENiagaraDataType::Curve)); }
-
-	void Merge(FNiagaraConstantMap &InMap)
-	{
-		for (auto MapIt = InMap.ScalarConstants.CreateIterator(); MapIt; ++MapIt)
-		{
-			SetOrAdd(MapIt.Key(), MapIt.Value());
-		}
-
-		for (auto MapIt = InMap.VectorConstants.CreateIterator(); MapIt; ++MapIt)
-		{
-			SetOrAdd(MapIt.Key(), MapIt.Value());
-		}
+	FNiagaraConstants_Vector(FName InName, const FVector4& InValue)
+	: FNiagaraConstantBase(InName)
+	, Value(InValue)
+	{}
 
-		for (auto MapIt = InMap.MatrixConstants.CreateIterator(); MapIt; ++MapIt)
-		{
-			SetOrAdd(MapIt.Key(), MapIt.Value());
-		}
+	UPROPERTY(EditAnywhere, Category = "Constant")
+	FVector4 Value;
+};
 
-		for (auto MapIt = InMap.DataConstants.CreateIterator(); MapIt; ++MapIt)
-		{
-			SetOrAdd(MapIt.Key(), MapIt.Value());
-		}
-	}
+USTRUCT()
+struct FNiagaraConstants_Matrix : public FNiagaraConstantBase
+{
+	GENERATED_USTRUCT_BODY()
 
-	const TMap<FNiagaraVariableInfo, UNiagaraDataObject*> &GetDataConstants()	const
-	{ 
-		return DataConstants; 
-	}
+	FNiagaraConstants_Matrix()
+	: Value(FMatrix::Identity)
+	{}
 
+	FNiagaraConstants_Matrix(FName InName, const FMatrix& InValue)
+	: FNiagaraConstantBase(InName)
+	, Value(InValue)
+	{}
 
-	virtual bool Serialize(FArchive &Ar)
-	{
-		// TODO: can't serialize the data object constants at the moment; need to figure that out
-		Ar << ScalarConstants << VectorConstants << MatrixConstants << DataConstants;
-		return true;
-	}
+	UPROPERTY(EditAnywhere, Category = "Constant")
+	FMatrix Value;
 };
 
+USTRUCT()
+struct FNiagaraConstants_DataObject : public FNiagaraConstantBase
+{
+	GENERATED_USTRUCT_BODY()
+
+	FNiagaraConstants_DataObject()
+	: Value(nullptr)
+	{}
 
+	FNiagaraConstants_DataObject(FName InName, UNiagaraDataObject* InValue)
+	: FNiagaraConstantBase(InName)
+	, Value(InValue)
+	{}
 
+	UPROPERTY(Instanced, VisibleAnywhere, Category = "Constant")
+	UNiagaraDataObject* Value;
+};
 
+//Dummy struct used to serialize in the old layout of FNiagaraConstants
+//This should be removed once everyone has recompiled and saved their scripts.
 USTRUCT()
-struct FNiagaraConstants
+struct FDeprecatedNiagaraConstants
 {
 	GENERATED_USTRUCT_BODY()
 
-private:
-
+	//DEPRECATED PROPERTIES. REMOVE SOON!
 	UPROPERTY()
-	TArray<FNiagaraVariableInfo> ScalarConstantsInfo;
+	TArray<FNiagaraVariableInfo> ScalarConstantsInfo_DEPRECATED;
 	UPROPERTY()
-	TArray<FNiagaraVariableInfo> VectorConstantsInfo;
+	TArray<FNiagaraVariableInfo> VectorConstantsInfo_DEPRECATED;
 	UPROPERTY()
-	TArray<FNiagaraVariableInfo> MatrixConstantsInfo;
+	TArray<FNiagaraVariableInfo> MatrixConstantsInfo_DEPRECATED;
 	//UPROPERTY()
-	TArray<FNiagaraVariableInfo> DataObjectConstantsInfo;
-
+	TArray<FNiagaraVariableInfo> DataObjectConstantsInfo_DEPRECATED;
 	UPROPERTY()
-	TArray<float> ScalarConstants;
+	TArray<float> ScalarConstants_DEPRECATED;
 	UPROPERTY()
-	TArray<FVector4> VectorConstants;
+	TArray<FVector4> VectorConstants_DEPRECATED;
 	UPROPERTY()
-	TArray<FMatrix> MatrixConstants;
+	TArray<FMatrix> MatrixConstants_DEPRECATED;
 	//UPROPERTY()
-	TArray<UNiagaraDataObject*> DataObjectConstants;
+	TArray<void*> DataObjectConstants_DEPRECATED;
+};
 
+USTRUCT()
+struct FNiagaraConstants
+{
+	GENERATED_USTRUCT_BODY()
+	
 public:
 
-	void Empty()
-	{
-		ScalarConstantsInfo.Empty();
-		VectorConstantsInfo.Empty();
-		MatrixConstantsInfo.Empty();
-		DataObjectConstantsInfo.Empty();
-		ScalarConstants.Empty();
-		VectorConstants.Empty();
-		MatrixConstants.Empty();
-		DataObjectConstants.Empty();
-	}
+	UPROPERTY(EditAnywhere, EditFixedSize, Category = "Constant")
+	TArray<FNiagaraConstants_Float> ScalarConstants;
+	UPROPERTY(EditAnywhere, EditFixedSize, Category = "Constant")
+	TArray<FNiagaraConstants_Vector> VectorConstants;
+	UPROPERTY(EditAnywhere, EditFixedSize, Category = "Constant")
+	TArray<FNiagaraConstants_Matrix> MatrixConstants;
+	UPROPERTY(EditAnywhere, EditFixedSize, Category = "Constant")
+	TArray<FNiagaraConstants_DataObject> DataObjectConstants;
 
-	void GetScalarConstant(int32 Index, float& OutValue, FNiagaraVariableInfo& OutInfo)const
-	{
-		OutInfo = ScalarConstantsInfo[Index]; 
-		OutValue = ScalarConstants[Index];
-	}
-	void GetVectorConstant(int32 Index, FVector4& OutValue, FNiagaraVariableInfo& OutInfo)const
-	{
-		OutInfo = VectorConstantsInfo[Index];
-		OutValue = VectorConstants[Index];
-	}
-	void GetMatrixConstant(int32 Index, FMatrix& OutValue, FNiagaraVariableInfo& OutInfo)const
-	{
-		OutInfo = MatrixConstantsInfo[Index];
-		OutValue = MatrixConstants[Index];
-	}
-	void GetDataObjectConstant(int32 Index, UNiagaraDataObject*& OutValue, FNiagaraVariableInfo& OutInfo)const
-	{
-		OutInfo = DataObjectConstantsInfo[Index];
-		OutValue = DataObjectConstants[Index];
-	}
+	NIAGARA_API void Empty();
+	NIAGARA_API void GetScalarConstant(int32 Index, float& OutValue, FNiagaraVariableInfo& OutInfo);
+	NIAGARA_API void GetVectorConstant(int32 Index, FVector4& OutValue, FNiagaraVariableInfo& OutInfo);
+	NIAGARA_API void GetMatrixConstant(int32 Index, FMatrix& OutValue, FNiagaraVariableInfo& OutInfo);
+	NIAGARA_API void GetDataObjectConstant(int32 Index, UNiagaraDataObject*& OutValue, FNiagaraVariableInfo& OutInfo);
 
-	const int32 GetNumScalarConstants()const{ return ScalarConstantsInfo.Num(); }
-	const int32 GetNumVectorConstants()const{ return VectorConstantsInfo.Num(); }
-	const int32 GetNumMatrixConstants()const{ return MatrixConstantsInfo.Num(); }
-	const int32 GetNumDataObjectConstants()const{ return DataObjectConstantsInfo.Num(); }
+	const int32 GetNumScalarConstants()const{ return ScalarConstants.Num(); }
+	const int32 GetNumVectorConstants()const{ return VectorConstants.Num(); }
+	const int32 GetNumMatrixConstants()const{ return MatrixConstants.Num(); }
+	const int32 GetNumDataObjectConstants()const{ return DataObjectConstants.Num(); }
 
-	/** Fills the entire constants set into the constant table. */
-	void AppendToConstantsTable(TArray<FVector4>& ConstantsTable)const
-	{
-		int32 Idx = ConstantsTable.Num();
-		int32 NewIdx = 0;
-		ConstantsTable.AddUninitialized(ScalarTableSize() + VectorTableSize() + MatrixTableSize());
-		for (float Sc : ScalarConstants)
-		{
-			ConstantsTable[Idx + (NewIdx / 4)][NewIdx % 4] = Sc;
-			++NewIdx;
-			Idx = NewIdx % 4 == 0 ? Idx + 1 : Idx;
-		}
+	NIAGARA_API void Init(class UNiagaraEmitterProperties* EmitterProps, struct FNiagaraEmitterScriptProperties* ScriptProps);
 
-		Idx += FMath::Min(1, NewIdx % 4);//Move to the next table entry if needed.
-		for (const FVector4& Vc : VectorConstants)
-		{
-			ConstantsTable[Idx++] = Vc;
-		}
-		for (const FMatrix& Mc : MatrixConstants)
-		{
-			ConstantsTable[Idx] = FVector4(Mc.M[0][0], Mc.M[0][1], Mc.M[0][2], Mc.M[0][3]);
-			ConstantsTable[Idx + 1] = FVector4(Mc.M[1][0], Mc.M[1][1], Mc.M[1][2], Mc.M[1][3]);
-			ConstantsTable[Idx + 2] = FVector4(Mc.M[2][0], Mc.M[2][1], Mc.M[2][2], Mc.M[2][3]);
-			ConstantsTable[Idx + 3] = FVector4(Mc.M[3][0], Mc.M[3][1], Mc.M[3][2], Mc.M[3][3]);
-			Idx += 4;
-		}
-	}
+	/** Fills the entire constants set into the constant table. */
+	void AppendToConstantsTable(TArray<FVector4>& ConstantsTable)const;
 
-	void AppendBufferConstants(TArray<class UNiagaraDataObject*> &Objs) const
-	{
-		Objs.Append(DataObjectConstants);
-	}
+	void AppendBufferConstants(TArray<class UNiagaraDataObject*> &Objs) const;
 
 	/** 
 	Fills only selected constants into the table. In the order they appear in the array of passed names not the order they appear in the set. 
 	Checks the passed map for entires to superceed the default values in the set.
 	*/
-	void AppendToConstantsTable(TArray<FVector4>& ConstantsTable, const FNiagaraConstantMap& Externals)const
-	{
-		int32 Idx = ConstantsTable.Num();
-		int32 NewIdx = 0;
-		ConstantsTable.AddUninitialized(ScalarTableSize() + VectorTableSize() + MatrixTableSize());
-		for (int32 i = 0; i < ScalarConstants.Num(); ++i)
-		{
-			const float* Scalar = Externals.FindScalar(ScalarConstantsInfo[i]);
-			ConstantsTable[Idx + (NewIdx / 4)][NewIdx % 4] = Scalar ? *Scalar : ScalarConstants[i];
-			++NewIdx;
-			Idx = NewIdx % 4 == 0 ? Idx + 1 : Idx;
-		}
-
-		Idx += FMath::Min(1, NewIdx % 4);//Move to the next table entry if needed.
-
-		for (int32 i = 0; i < VectorConstants.Num(); ++i)
-		{
-			const FVector4* Vector = Externals.FindVector(VectorConstantsInfo[i]);
-			ConstantsTable[Idx++] = Vector ? *Vector : VectorConstants[i];
-		}
-
-		for (int32 i = 0; i < MatrixConstants.Num(); ++i)
-		{
-			const FMatrix* Mat = Externals.FindMatrix(MatrixConstantsInfo[i]);
-			FMatrix Mc = Mat ? *Mat : MatrixConstants[i];
-			ConstantsTable[Idx] = FVector4(Mc.M[0][0], Mc.M[0][1], Mc.M[0][2], Mc.M[0][3]);
-			ConstantsTable[Idx + 1] = FVector4(Mc.M[1][0], Mc.M[1][1], Mc.M[1][2], Mc.M[1][3]);
-			ConstantsTable[Idx + 2] = FVector4(Mc.M[2][0], Mc.M[2][1], Mc.M[2][2], Mc.M[2][3]);
-			ConstantsTable[Idx + 3] = FVector4(Mc.M[3][0], Mc.M[3][1], Mc.M[3][2], Mc.M[3][3]);
-		}
-	}
-
-	void AppendExternalBufferConstants(TArray<class UNiagaraDataObject*> &DataObjs, const FNiagaraConstantMap& Externals) const
-	{
-		for (int32 i = 0; i < DataObjectConstantsInfo.Num(); ++i)
-		{
-			FNiagaraVariableInfo CcID = DataObjectConstantsInfo[i];
-			UNiagaraDataObject* Data = Externals.FindDataObj(CcID);
-			DataObjs.Add(Data);
-		}
-	}
-
-	void SetOrAdd(const FNiagaraVariableInfo& Constant, float Value)
-	{
-		int32 Idx = ScalarConstantsInfo.Find(Constant);
-		if (Idx == INDEX_NONE)
-		{
-			ScalarConstantsInfo.Add(Constant);
-			ScalarConstants.Add(Value);
-		}
-		else
-		{
-			ScalarConstants[Idx] = Value;
-		}
-	}
-
-	void SetOrAdd(const FNiagaraVariableInfo& Constant, const FVector4& Value)
-	{
-		int32 Idx = VectorConstantsInfo.Find(Constant);
-		if (Idx == INDEX_NONE)
-		{
-			VectorConstantsInfo.Add(Constant);
-			VectorConstants.Add(Value);
-		}
-		else
-		{
-			VectorConstants[Idx] = Value;
-		}
-	}
-
-	void SetOrAdd(const FNiagaraVariableInfo& Constant, const FMatrix& Value)
-	{
-		int32 Idx = MatrixConstantsInfo.Find(Constant);
-		if (Idx == INDEX_NONE)
-		{
-			MatrixConstantsInfo.Add(Constant);
-			MatrixConstants.Add(Value);
-		}
-		else
-		{
-			MatrixConstants[Idx] = Value;
-		}
-	}
-
-	void SetOrAdd(const FNiagaraVariableInfo& Constant, UNiagaraDataObject* Value)
-	{
-		int32 Idx = DataObjectConstantsInfo.Find(Constant);
-		if (Idx == INDEX_NONE)
-		{
-			DataObjectConstantsInfo.Add(Constant);
-			DataObjectConstants.Add(Value);
-		}
-		else
-		{
-			DataObjectConstants[Idx] = Value;
-		}
-	}
+	void AppendToConstantsTable(TArray<FVector4>& ConstantsTable, const FNiagaraConstantMap& Externals)const;
+	void AppendExternalBufferConstants(TArray<class UNiagaraDataObject*> &DataObjs, const FNiagaraConstantMap& Externals) const;
+
+	NIAGARA_API void SetOrAdd(const FNiagaraVariableInfo& Constant, float Value);
+	NIAGARA_API void SetOrAdd(const FNiagaraVariableInfo& Constant, const FVector4& Value);
+	NIAGARA_API void SetOrAdd(const FNiagaraVariableInfo& Constant, const FMatrix& Value);
+	NIAGARA_API void SetOrAdd(const FNiagaraVariableInfo& Constant, UNiagaraDataObject* Value);
+	NIAGARA_API void SetOrAdd(FName Name, float Value);
+	NIAGARA_API void SetOrAdd(FName Name, const FVector4& Value);
+	NIAGARA_API void SetOrAdd(FName Name, const FMatrix& Value);
+	NIAGARA_API void SetOrAdd(FName Name, UNiagaraDataObject* Value);
 
 	FORCEINLINE int32 NumScalars()const { return ScalarConstants.Num(); }
 	FORCEINLINE int32 NumVectors()const { return VectorConstants.Num(); }
@@ -341,11 +227,16 @@ public:
 
 	/**	Returns the number of vectors these constants would use in a constants table. */
 	FORCEINLINE int32 GetTableSize()const { return ScalarTableSize() + VectorTableSize() + MatrixTableSize(); }
+	
+	NIAGARA_API float* FindScalar(FName Name);
+	NIAGARA_API FVector4* FindVector(FName Name);
+	NIAGARA_API FMatrix* FindMatrix(FName Name);
+	NIAGARA_API UNiagaraDataObject* FindDataObj(FName Name);
 
 	/** Return the first constant used for this constant in a table. */
 	FORCEINLINE int32 GetTableIndex_Scalar(const FNiagaraVariableInfo& Constant)const
 	{ 
-		int Idx = ScalarConstantsInfo.Find(Constant);
+		int Idx = ScalarConstants.IndexOfByPredicate([&](const FNiagaraConstants_Float& V){ return Constant.Name == V.Name; });
 		if (Idx != INDEX_NONE)
 		{
 			return Idx / 4;
@@ -355,7 +246,7 @@ public:
 
 	FORCEINLINE int32 GetTableIndex_Vector(const FNiagaraVariableInfo& Constant)const
 	{ 
-		int Idx = VectorConstantsInfo.Find(Constant);
+		int Idx = VectorConstants.IndexOfByPredicate([&](const FNiagaraConstants_Vector& V){ return Constant.Name == V.Name; });
 		if (Idx != INDEX_NONE)
 		{
 			return ScalarTableSize() + Idx;
@@ -365,7 +256,7 @@ public:
 
 	FORCEINLINE int32 GetTableIndex_Matrix(const FNiagaraVariableInfo& Constant)const
 	{ 
-		int Idx = MatrixConstantsInfo.Find(Constant);
+		int Idx = MatrixConstants.IndexOfByPredicate([&](const FNiagaraConstants_Matrix& V){ return Constant.Name == V.Name; });
 		if (Idx != INDEX_NONE)
 		{
 			return ScalarTableSize() + VectorTableSize() + Idx * 4;
@@ -375,10 +266,9 @@ public:
 
 	FORCEINLINE int32 GetTableIndex_DataObj(const FNiagaraVariableInfo& Constant)const
 	{ 
-		return DataObjectConstantsInfo.Find(Constant); 
+		return DataObjectConstants.IndexOfByPredicate([&](const FNiagaraConstants_DataObject& V){ return Constant.Name == V.Name; });
 	}
 
 	/** Return the absolute index of the passed scalar. Can be used to get the component index also. */
-	FORCEINLINE int32 GetAbsoluteIndex_Scalar(const FNiagaraVariableInfo& Constant)const{ return ScalarConstantsInfo.Find(Constant); }
+	FORCEINLINE int32 GetAbsoluteIndex_Scalar(const FNiagaraVariableInfo& Constant)const{ return ScalarConstants.IndexOfByPredicate([&](const FNiagaraConstants_Float& V){ return Constant.Name == V.Name; });; }
 };
-
diff --git a/Engine/Source/Runtime/Niagara/NiagaraScriptConstantData.h b/Engine/Source/Runtime/Niagara/NiagaraScriptConstantData.h
--- a/Engine/Source/Runtime/Niagara/NiagaraScriptConstantData.h
+++ b/Engine/Source/Runtime/Niagara/NiagaraScriptConstantData.h
@@ -13,7 +13,10 @@ struct FNiagaraScriptConstantData
 
 	/** The set of external constants for this script. */
 	UPROPERTY()
-	FNiagaraConstants ExternalConstants;
+	FDeprecatedNiagaraConstants ExternalConstants_DEPRECATED;
+
+	UPROPERTY()
+	FNiagaraConstants ExposedConstants;
 
 	/** The set of internal constants for this script. */
 	UPROPERTY()
@@ -21,7 +24,7 @@ struct FNiagaraScriptConstantData
 
 	void Empty()
 	{
-		ExternalConstants.Empty();
+		ExposedConstants.Empty();
 		InternalConstants.Empty();
 	}
 
@@ -31,8 +34,8 @@ struct FNiagaraScriptConstantData
 		//First up in the table comes the External constants in scalar->vector->matrix order.
 		//Only fills the constants actually used by the script.
 		OutConstantTable.Empty();
-		ExternalConstants.AppendToConstantsTable(OutConstantTable, ExternalConstantsMap);
-		ExternalConstants.AppendExternalBufferConstants(DataObjTable, ExternalConstantsMap);
+		ExposedConstants.AppendToConstantsTable(OutConstantTable, ExternalConstantsMap);
+		ExposedConstants.AppendExternalBufferConstants(DataObjTable, ExternalConstantsMap);
 		//Next up add all the internal constants from the script.
 		InternalConstants.AppendToConstantsTable(OutConstantTable);
 		InternalConstants.AppendBufferConstants(DataObjTable);
@@ -47,7 +50,7 @@ struct FNiagaraScriptConstantData
 	template<typename T>
 	void SetOrAddExternal(const FNiagaraVariableInfo& Constant, const T& Value)
 	{
-		ExternalConstants.SetOrAdd(Constant, Value);
+		ExposedConstants.SetOrAdd(Constant, Value);
 	}
 
 	/**
@@ -56,11 +59,11 @@ struct FNiagaraScriptConstantData
 	*/
 	void GetTableIndex(const FNiagaraVariableInfo& Constant, bool bInternal, int32& OutConstantIdx, int32& OutComponentIndex, ENiagaraDataType& OutType)const
 	{
-		int32 Base = bInternal ? ExternalConstants.GetTableSize() : 0;
+		int32 Base = bInternal ? ExposedConstants.GetTableSize() : 0;
 		int32 ConstIdx = INDEX_NONE;
 		int32 CompIdx = INDEX_NONE;
 		ENiagaraDataType Type;
-		const FNiagaraConstants& Consts = bInternal ? InternalConstants : ExternalConstants;
+		const FNiagaraConstants& Consts = bInternal ? InternalConstants : ExposedConstants;
 
 		ConstIdx = Consts.GetAbsoluteIndex_Scalar(Constant);
 		Type = ENiagaraDataType::Scalar;
diff --git a/Engine/Source/Runtime/Niagara/Private/NiagaraConstantSet.cpp b/Engine/Source/Runtime/Niagara/Private/NiagaraConstantSet.cpp
--- a/Engine/Source/Runtime/Niagara/Private/NiagaraConstantSet.cpp
+++ b/Engine/Source/Runtime/Niagara/Private/NiagaraConstantSet.cpp
@@ -3,12 +3,412 @@
 #include "NiagaraPrivate.h"
 #include "NiagaraConstantSet.h"
 
-template<>
-struct TStructOpsTypeTraits<FNiagaraConstantMap> : public TStructOpsTypeTraitsBase
+//////////////////////////////////////////////////////////////////////////
+// FNiagaraConstantMap
+
+void FNiagaraConstantMap::Empty()
+{
+	ScalarConstants.Empty();
+	VectorConstants.Empty();
+	MatrixConstants.Empty();
+}
+
+void FNiagaraConstantMap::SetOrAdd(FNiagaraVariableInfo ID, float Sc)
+{
+	check(ID.Type == ENiagaraDataType::Scalar);
+	ScalarConstants.FindOrAdd(ID) = Sc;
+}
+
+void FNiagaraConstantMap::SetOrAdd(FNiagaraVariableInfo ID, const FVector4& Vc)
+{
+	check(ID.Type == ENiagaraDataType::Vector);
+	VectorConstants.FindOrAdd(ID) = Vc;
+}
+
+void FNiagaraConstantMap::SetOrAdd(FNiagaraVariableInfo ID, const FMatrix& Mc)
+{
+	check(ID.Type == ENiagaraDataType::Matrix);
+	MatrixConstants.FindOrAdd(ID) = Mc;
+}
+
+void FNiagaraConstantMap::SetOrAdd(FNiagaraVariableInfo ID, UNiagaraDataObject* Cc)
+{
+	check(ID.Type == ENiagaraDataType::Curve);
+	DataConstants.FindOrAdd(ID) = Cc;
+}
+
+void FNiagaraConstantMap::SetOrAdd(FName Name, float Sc)
+{
+	ScalarConstants.FindOrAdd(FNiagaraVariableInfo(Name, ENiagaraDataType::Scalar)) = Sc;
+}
+
+void FNiagaraConstantMap::SetOrAdd(FName Name, const FVector4& Vc)
+{
+	VectorConstants.FindOrAdd(FNiagaraVariableInfo(Name, ENiagaraDataType::Vector)) = Vc;
+}
+
+void FNiagaraConstantMap::SetOrAdd(FName Name, const FMatrix& Mc)
 {
-	enum
+	MatrixConstants.FindOrAdd(FNiagaraVariableInfo(Name, ENiagaraDataType::Matrix)) = Mc;
+}
+
+void FNiagaraConstantMap::SetOrAdd(FName Name, class UNiagaraDataObject* Cc)
+{
+	DataConstants.FindOrAdd(FNiagaraVariableInfo(Name, ENiagaraDataType::Curve)) = Cc;
+}
+
+float* FNiagaraConstantMap::FindScalar(FNiagaraVariableInfo ID){ return ScalarConstants.Find(ID); }
+FVector4* FNiagaraConstantMap::FindVector(FNiagaraVariableInfo ID){ return VectorConstants.Find(ID); }
+FMatrix* FNiagaraConstantMap::FindMatrix(FNiagaraVariableInfo ID){ return MatrixConstants.Find(ID); }
+UNiagaraDataObject* FNiagaraConstantMap::FindDataObj(FNiagaraVariableInfo ID)
+{
+	UNiagaraDataObject * const *Ptr = DataConstants.Find(ID);
+	if (Ptr)
+	{
+		return *Ptr;
+	}
+
+	return nullptr;
+}
+const float* FNiagaraConstantMap::FindScalar(FNiagaraVariableInfo ID)const { return ScalarConstants.Find(ID); }
+const FVector4* FNiagaraConstantMap::FindVector(FNiagaraVariableInfo ID)const { return VectorConstants.Find(ID); }
+const FMatrix* FNiagaraConstantMap::FindMatrix(FNiagaraVariableInfo ID)const { return MatrixConstants.Find(ID); }
+UNiagaraDataObject* FNiagaraConstantMap::FindDataObj(FNiagaraVariableInfo ID) const
+{
+	UNiagaraDataObject * const *Ptr = DataConstants.Find(ID);
+	if (Ptr)
+	{
+		return *Ptr;
+	}
+
+	return nullptr;
+}
+float* FNiagaraConstantMap::FindScalar(FName Name) { return ScalarConstants.Find(FNiagaraVariableInfo(Name, ENiagaraDataType::Scalar)); }
+FVector4* FNiagaraConstantMap::FindVector(FName Name) { return VectorConstants.Find(FNiagaraVariableInfo(Name, ENiagaraDataType::Vector)); }
+FMatrix* FNiagaraConstantMap::FindMatrix(FName Name) { return MatrixConstants.Find(FNiagaraVariableInfo(Name, ENiagaraDataType::Matrix)); }
+UNiagaraDataObject* FNiagaraConstantMap::FindDataObj(FName Name){ return FindDataObj(FNiagaraVariableInfo(Name, ENiagaraDataType::Curve)); }
+
+const float* FNiagaraConstantMap::FindScalar(FName Name)const  { return ScalarConstants.Find(FNiagaraVariableInfo(Name, ENiagaraDataType::Scalar)); }
+const FVector4* FNiagaraConstantMap::FindVector(FName Name)const  { return VectorConstants.Find(FNiagaraVariableInfo(Name, ENiagaraDataType::Vector)); }
+const FMatrix* FNiagaraConstantMap::FindMatrix(FName Name)const  { return MatrixConstants.Find(FNiagaraVariableInfo(Name, ENiagaraDataType::Matrix)); }
+UNiagaraDataObject* FNiagaraConstantMap::FindDataObj(FName Name)const { return FindDataObj(FNiagaraVariableInfo(Name, ENiagaraDataType::Curve)); }
+
+void FNiagaraConstantMap::Merge(FNiagaraConstants &InConstants)
+{
+	for (FNiagaraConstants_Float& C : InConstants.ScalarConstants)
+	{
+		SetOrAdd(FNiagaraVariableInfo(C.Name, ENiagaraDataType::Scalar), C.Value);
+	}
+	for (FNiagaraConstants_Vector& C : InConstants.VectorConstants)
+	{
+		SetOrAdd(FNiagaraVariableInfo(C.Name, ENiagaraDataType::Vector), C.Value);
+	}
+	for (FNiagaraConstants_Matrix& C : InConstants.MatrixConstants)
+	{
+		SetOrAdd(FNiagaraVariableInfo(C.Name, ENiagaraDataType::Matrix), C.Value);
+	}
+	for (FNiagaraConstants_DataObject& C : InConstants.DataObjectConstants)
 	{
-		WithSerializer = true
-	};
-};
-IMPLEMENT_STRUCT(NiagaraConstantMap);
+		SetOrAdd(FNiagaraVariableInfo(C.Name, ENiagaraDataType::Curve), C.Value);
+	}
+}
+
+void FNiagaraConstantMap::Merge(FNiagaraConstantMap &InMap)
+{
+	for (auto MapIt = InMap.ScalarConstants.CreateIterator(); MapIt; ++MapIt)
+	{
+		SetOrAdd(MapIt.Key(), MapIt.Value());
+	}
+
+	for (auto MapIt = InMap.VectorConstants.CreateIterator(); MapIt; ++MapIt)
+	{
+		SetOrAdd(MapIt.Key(), MapIt.Value());
+	}
+
+	for (auto MapIt = InMap.MatrixConstants.CreateIterator(); MapIt; ++MapIt)
+	{
+		SetOrAdd(MapIt.Key(), MapIt.Value());
+	}
+
+	for (auto MapIt = InMap.DataConstants.CreateIterator(); MapIt; ++MapIt)
+	{
+		SetOrAdd(MapIt.Key(), MapIt.Value());
+	}
+}
+
+const TMap<FNiagaraVariableInfo, class UNiagaraDataObject*> &FNiagaraConstantMap::GetDataConstants()	const
+{
+	return DataConstants;
+}
+
+//////////////////////////////////////////////////////////////////////////
+
+void FNiagaraConstants::Empty()
+{
+	ScalarConstants.Empty();
+	VectorConstants.Empty();
+	MatrixConstants.Empty();
+	DataObjectConstants.Empty();
+}
+
+void FNiagaraConstants::GetScalarConstant(int32 Index, float& OutValue, FNiagaraVariableInfo& OutInfo)
+{
+	OutInfo = FNiagaraVariableInfo(ScalarConstants[Index].Name, ENiagaraDataType::Scalar);
+	OutValue = ScalarConstants[Index].Value;
+}
+
+void FNiagaraConstants::GetVectorConstant(int32 Index, FVector4& OutValue, FNiagaraVariableInfo& OutInfo)
+{
+	OutInfo = FNiagaraVariableInfo(VectorConstants[Index].Name, ENiagaraDataType::Vector);
+	OutValue = VectorConstants[Index].Value;
+}
+
+void FNiagaraConstants::GetMatrixConstant(int32 Index, FMatrix& OutValue, FNiagaraVariableInfo& OutInfo)
+{
+	OutInfo = FNiagaraVariableInfo(MatrixConstants[Index].Name, ENiagaraDataType::Matrix);
+	OutValue = MatrixConstants[Index].Value;
+}
+
+void FNiagaraConstants::GetDataObjectConstant(int32 Index, UNiagaraDataObject*& OutValue, FNiagaraVariableInfo& OutInfo)
+{
+	OutInfo = FNiagaraVariableInfo(DataObjectConstants[Index].Name, ENiagaraDataType::Curve);
+	OutValue = DataObjectConstants[Index].Value;
+}
+
+void FNiagaraConstants::Init(UNiagaraEmitterProperties* EmitterProps, FNiagaraEmitterScriptProperties* ScriptProps)
+{
+	check(ScriptProps);
+	check(EmitterProps);
+
+	//Copy the script constants, retaining any altered values.
+	FNiagaraConstants& ScriptConstants = ScriptProps->Script->ConstantData.ExposedConstants;
+
+	TArray<FNiagaraConstants_Float> NewScalarConstants = ScriptConstants.ScalarConstants;
+	for (auto& Override : ScalarConstants)
+	{
+		int32 Idx = NewScalarConstants.IndexOfByPredicate([&](const FNiagaraConstants_Float& C){ return Override.Name == C.Name; });
+
+		if (Idx != INDEX_NONE)
+			NewScalarConstants[Idx].Value = Override.Value;
+	}
+	ScalarConstants = NewScalarConstants;
+
+	TArray<FNiagaraConstants_Vector> NewVectorConstants = ScriptConstants.VectorConstants;
+	for (auto& Override : VectorConstants)
+	{
+		int32 Idx = NewVectorConstants.IndexOfByPredicate([&](const FNiagaraConstants_Vector& C){ return Override.Name == C.Name; });
+
+		if (Idx != INDEX_NONE)
+			NewVectorConstants[Idx].Value = Override.Value;
+	}
+	VectorConstants = NewVectorConstants;
+
+	TArray<FNiagaraConstants_Matrix> NewMatrixConstants = ScriptConstants.MatrixConstants;
+	for (auto& Override : MatrixConstants)
+	{
+		int32 Idx = NewMatrixConstants.IndexOfByPredicate([&](const FNiagaraConstants_Matrix& C){ return Override.Name == C.Name; });
+
+		if (Idx != INDEX_NONE)
+			NewMatrixConstants[Idx].Value = Override.Value;
+
+	}
+	MatrixConstants = NewMatrixConstants;
+
+	TArray<FNiagaraConstants_DataObject>& ScriptDataObjectConstants = ScriptConstants.DataObjectConstants;
+	TArray<FNiagaraConstants_DataObject> OldDataObjectConstants = DataObjectConstants;
+	DataObjectConstants.Empty();
+	for (FNiagaraConstants_DataObject& ScriptConst : ScriptDataObjectConstants)
+	{
+		int32 OverrideIdx = OldDataObjectConstants.IndexOfByPredicate([&](const FNiagaraConstants_DataObject& C){ return ScriptConst.Name == C.Name && ScriptConst.Value && C.Value &&ScriptConst.Value->GetClass() == C.Value->GetClass(); });
+		
+		int32 AddIndex = DataObjectConstants.AddZeroed();
+		DataObjectConstants[AddIndex].Name = ScriptDataObjectConstants[AddIndex].Name;
+		
+		if (OverrideIdx != INDEX_NONE)
+		{
+			//If this constant already has data in the effect, use that.
+			DataObjectConstants[AddIndex].Value = OldDataObjectConstants[OverrideIdx].Value;
+		}		
+		else
+		{
+			//Otherwise, duplicate the data from the script.
+			DataObjectConstants[AddIndex].Value = CastChecked<UNiagaraDataObject>(StaticDuplicateObject(ScriptConst.Value, EmitterProps, NULL));
+		}
+	}
+}
+
+void FNiagaraConstants::AppendToConstantsTable(TArray<FVector4>& ConstantsTable)const
+{
+	int32 Idx = ConstantsTable.Num();
+	int32 NewIdx = 0;
+	ConstantsTable.AddUninitialized(ScalarTableSize() + VectorTableSize() + MatrixTableSize());
+	for (const FNiagaraConstants_Float& Sc : ScalarConstants)
+	{
+		ConstantsTable[Idx + (NewIdx / 4)][NewIdx % 4] = Sc.Value;
+		++NewIdx;
+		Idx = NewIdx % 4 == 0 ? Idx + 1 : Idx;
+	}
+
+	Idx += FMath::Min(1, NewIdx % 4);//Move to the next table entry if needed.
+	for (const FNiagaraConstants_Vector& Vc : VectorConstants)
+	{
+		ConstantsTable[Idx++] = Vc.Value;
+	}
+	for (const FNiagaraConstants_Matrix& C : MatrixConstants)
+	{
+		const FMatrix& Mc = C.Value;
+		ConstantsTable[Idx] = FVector4(Mc.M[0][0], Mc.M[0][1], Mc.M[0][2], Mc.M[0][3]);
+		ConstantsTable[Idx + 1] = FVector4(Mc.M[1][0], Mc.M[1][1], Mc.M[1][2], Mc.M[1][3]);
+		ConstantsTable[Idx + 2] = FVector4(Mc.M[2][0], Mc.M[2][1], Mc.M[2][2], Mc.M[2][3]);
+		ConstantsTable[Idx + 3] = FVector4(Mc.M[3][0], Mc.M[3][1], Mc.M[3][2], Mc.M[3][3]);
+		Idx += 4;
+	}
+}
+
+void FNiagaraConstants::AppendBufferConstants(TArray<class UNiagaraDataObject*> &Objs) const
+{
+	for (const FNiagaraConstants_DataObject& C : DataObjectConstants)
+	{
+		Objs.Add(C.Value);
+	}
+}
+
+void FNiagaraConstants::AppendToConstantsTable(TArray<FVector4>& ConstantsTable, const FNiagaraConstantMap& Externals)const
+{
+	int32 Idx = ConstantsTable.Num();
+	int32 NewIdx = 0;
+	ConstantsTable.AddUninitialized(ScalarTableSize() + VectorTableSize() + MatrixTableSize());
+	for (int32 i = 0; i < ScalarConstants.Num(); ++i)
+	{
+		const float* Scalar = Externals.FindScalar(ScalarConstants[i].Name);
+		ConstantsTable[Idx + (NewIdx / 4)][NewIdx % 4] = Scalar ? *Scalar : ScalarConstants[i].Value;
+		++NewIdx;
+		Idx = NewIdx % 4 == 0 ? Idx + 1 : Idx;
+	}
+
+	Idx += FMath::Min(1, NewIdx % 4);//Move to the next table entry if needed.
+
+	for (int32 i = 0; i < VectorConstants.Num(); ++i)
+	{
+		const FVector4* Vector = Externals.FindVector(VectorConstants[i].Name);
+		ConstantsTable[Idx++] = Vector ? *Vector : VectorConstants[i].Value;
+	}
+
+	for (int32 i = 0; i < MatrixConstants.Num(); ++i)
+	{
+		const FMatrix* Mat = Externals.FindMatrix(MatrixConstants[i].Name);
+		FMatrix Mc = Mat ? *Mat : MatrixConstants[i].Value;
+		ConstantsTable[Idx] = FVector4(Mc.M[0][0], Mc.M[0][1], Mc.M[0][2], Mc.M[0][3]);
+		ConstantsTable[Idx + 1] = FVector4(Mc.M[1][0], Mc.M[1][1], Mc.M[1][2], Mc.M[1][3]);
+		ConstantsTable[Idx + 2] = FVector4(Mc.M[2][0], Mc.M[2][1], Mc.M[2][2], Mc.M[2][3]);
+		ConstantsTable[Idx + 3] = FVector4(Mc.M[3][0], Mc.M[3][1], Mc.M[3][2], Mc.M[3][3]);
+	}
+}
+
+void FNiagaraConstants::AppendExternalBufferConstants(TArray<class UNiagaraDataObject*> &DataObjs, const FNiagaraConstantMap& Externals) const
+{
+	for (int32 i = 0; i < DataObjectConstants.Num(); ++i)
+	{
+		UNiagaraDataObject* Data = Externals.FindDataObj(DataObjectConstants[i].Name);
+		DataObjs.Add(Data);
+	}
+}
+
+void FNiagaraConstants::SetOrAdd(const FNiagaraVariableInfo& Constant, float Value)
+{
+	check(Constant.Type == ENiagaraDataType::Scalar);
+	FNiagaraConstants_Float* Const = ScalarConstants.FindByPredicate([&](const FNiagaraConstants_Float& C){ return Constant.Name == C.Name; });
+	if (Const)
+	{
+		Const->Value = Value;
+	}
+	else
+	{
+		ScalarConstants.Add(FNiagaraConstants_Float(Constant.Name, Value));
+	}
+}
+
+void FNiagaraConstants::SetOrAdd(const FNiagaraVariableInfo& Constant, const FVector4& Value)
+{
+	check(Constant.Type == ENiagaraDataType::Vector);
+	FNiagaraConstants_Vector* Const = VectorConstants.FindByPredicate([&](const FNiagaraConstants_Vector& C){ return Constant.Name == C.Name; });
+	if (Const)
+	{
+		Const->Value = Value;
+	}
+	else
+	{
+		VectorConstants.Add(FNiagaraConstants_Vector(Constant.Name, Value));
+	}
+}
+
+void FNiagaraConstants::SetOrAdd(const FNiagaraVariableInfo& Constant, const FMatrix& Value)
+{
+	check(Constant.Type == ENiagaraDataType::Matrix);
+	FNiagaraConstants_Matrix* Const = MatrixConstants.FindByPredicate([&](const FNiagaraConstants_Matrix& C){ return Constant.Name == C.Name; });
+	if (Const)
+	{
+		Const->Value = Value;
+	}
+	else
+	{
+		MatrixConstants.Add(FNiagaraConstants_Matrix(Constant.Name, Value));
+	}
+}
+
+void FNiagaraConstants::SetOrAdd(const FNiagaraVariableInfo& Constant, UNiagaraDataObject* Value)
+{
+	check(Constant.Type == ENiagaraDataType::Curve);
+	FNiagaraConstants_DataObject* Const = DataObjectConstants.FindByPredicate([&](const FNiagaraConstants_DataObject& C){ return Constant.Name == C.Name; });
+	if (Const)
+	{
+		Const->Value = Value;
+	}
+	else
+	{
+		DataObjectConstants.Add(FNiagaraConstants_DataObject(Constant.Name, Value));
+	}
+}
+
+void FNiagaraConstants::SetOrAdd(FName Name, float Value)
+{
+	SetOrAdd(FNiagaraVariableInfo(Name, ENiagaraDataType::Scalar), Value);
+}
+
+void FNiagaraConstants::SetOrAdd(FName Name, const FVector4& Value)
+{
+	SetOrAdd(FNiagaraVariableInfo(Name, ENiagaraDataType::Vector), Value);
+}
+
+void FNiagaraConstants::SetOrAdd(FName Name, const FMatrix& Value)
+{
+	SetOrAdd(FNiagaraVariableInfo(Name, ENiagaraDataType::Matrix), Value);
+}
+
+void FNiagaraConstants::SetOrAdd(FName Name, UNiagaraDataObject* Value)
+{
+	SetOrAdd(FNiagaraVariableInfo(Name, ENiagaraDataType::Curve), Value);
+}
+
+float* FNiagaraConstants::FindScalar(FName Name)
+{
+	auto* C = ScalarConstants.FindByPredicate([&](const FNiagaraConstants_Float& C){ return C.Name == Name; });
+	return C ? &C->Value : nullptr;
+}
+FVector4* FNiagaraConstants::FindVector(FName Name)
+{
+	auto* C = VectorConstants.FindByPredicate([&](const FNiagaraConstants_Vector& C){ return C.Name == Name; });
+	return C ? &C->Value : nullptr;
+}
+FMatrix* FNiagaraConstants::FindMatrix(FName Name)
+{
+	auto* C = MatrixConstants.FindByPredicate([&](const FNiagaraConstants_Matrix& C){ return C.Name == Name; });
+	return C ? &C->Value : nullptr;
+}
+UNiagaraDataObject* FNiagaraConstants::FindDataObj(FName Name)
+{
+	auto* C = DataObjectConstants.FindByPredicate([&](const FNiagaraConstants_DataObject& C){ return C.Name == Name; });
+	return C ? C->Value : nullptr;
+}
+
+//////////////////////////////////////////////////////////////////////////
diff --git a/Engine/Source/Runtime/Niagara/Private/NiagaraEffect.cpp b/Engine/Source/Runtime/Niagara/Private/NiagaraEffect.cpp
--- a/Engine/Source/Runtime/Niagara/Private/NiagaraEffect.cpp
+++ b/Engine/Source/Runtime/Niagara/Private/NiagaraEffect.cpp
@@ -13,13 +13,20 @@ UNiagaraEffect::UNiagaraEffect(const FObjectInitializer& ObjectInitializer)
 
 
 
-FNiagaraEmitterProperties *UNiagaraEffect::AddEmitterProperties()
+UNiagaraEmitterProperties *UNiagaraEffect::AddEmitterProperties(UNiagaraEmitterProperties *Props)
 {
-	FNiagaraEmitterProperties *Props = new FNiagaraEmitterProperties();
+	if (Props == nullptr)
+	{
+		Props = NewObject<UNiagaraEmitterProperties>(this);
+	}
 	EmitterProps.Add(Props);
 	return Props;
 }
 
+void UNiagaraEffect::DeleteEmitterProperties(UNiagaraEmitterProperties *Props)
+{
+	EmitterProps.Remove(Props);
+}
 
 
 void UNiagaraEffect::CreateEffectRendererProps(TSharedPtr<FNiagaraSimulation> Sim)
@@ -35,37 +42,11 @@ void UNiagaraEffect::CreateEffectRendererProps(TSharedPtr<FNiagaraSimulation> Si
 	}
 }
 
-
-void UNiagaraEffect::PostLoad()
-{
-	Super::PostLoad();
-	EmitterProps.Empty();
-	for (FNiagaraEmitterProperties &Props : EmitterPropsSerialized)
-	{
-		FNiagaraEmitterProperties *NewProps = new FNiagaraEmitterProperties(Props);
-		EmitterProps.Add(NewProps);
-	}
-}
-
-
-void UNiagaraEffect::PreSave()
-{
-	EmitterPropsSerialized.Empty();
-	for (FNiagaraEmitterProperties *Props : EmitterProps)
-	{
-		EmitterPropsSerialized.Add(*Props);
-	}
-}
-
-
-
-
-TSharedPtr<FNiagaraSimulation> FNiagaraEffectInstance::AddEmitter(FNiagaraEmitterProperties *Properties)
+TSharedPtr<FNiagaraSimulation> FNiagaraEffectInstance::AddEmitter(UNiagaraEmitterProperties *Properties)
 {
 	FNiagaraSimulation *SimPtr = new FNiagaraSimulation(Properties, Effect);
 	TSharedPtr<FNiagaraSimulation> Sim = MakeShareable(SimPtr);
-
-	Sim->SetRenderModuleType(RMT_Sprites, Component->GetWorld()->FeatureLevel);
+	Sim->SetRenderModuleType(Properties->RenderModuleType, Component->GetWorld()->FeatureLevel);
 	Emitters.Add(Sim);
 	FNiagaraSceneProxy *SceneProxy = static_cast<FNiagaraSceneProxy*>(Component->SceneProxy);
 	SceneProxy->UpdateEffectRenderers(this);
@@ -73,6 +54,66 @@ TSharedPtr<FNiagaraSimulation> FNiagaraEffectInstance::AddEmitter(FNiagaraEmitte
 	return Sim;
 }
 
+void FNiagaraEffectInstance::DeleteEmitter(TSharedPtr<FNiagaraSimulation> Emitter)
+{
+	Emitters.Remove(Emitter);
+	FNiagaraSceneProxy *SceneProxy = static_cast<FNiagaraSceneProxy*>(Component->SceneProxy);
+	SceneProxy->UpdateEffectRenderers(this);
+}
+
+void FNiagaraEffectInstance::Tick(float DeltaSeconds)
+{
+	Constants.SetOrAdd(FName(TEXT("EffectGrid")), VolumeGrid);
+
+	// pass the constants down to the emitter
+	// TODO: should probably just pass a pointer to the table
+	EffectBounds.Init();
+
+	for (TSharedPtr<FNiagaraSimulation>&it : Emitters)
+	{
+		UNiagaraEmitterProperties *Props = it->GetProperties().Get();
+		check(Props);
+
+		int Duration = Props->EndTime - Props->StartTime;
+		int LoopedStartTime = Props->StartTime + Duration*it->GetLoopCount();
+		int LoopedEndTime = Props->EndTime + Duration*it->GetLoopCount();
+
+		// manage emitter lifetime
+		//
+		if ((Props->StartTime == 0.0f && Props->EndTime == 0.0f)
+			|| (LoopedStartTime<Age && LoopedEndTime>Age)
+			)
+		{
+			it->SetTickState(NTS_Running);
+		}
+		else
+		{
+			// if we're past end time, manage looping; we reset the emitters age constant
+			// if it has one
+			if (Props->NumLoops > 1 && it->GetLoopCount() < Props->NumLoops)
+			{
+				it->LoopRestart();
+			}
+			else
+			{
+				it->SetTickState(NTS_Dieing);
+			}
+		}
+
+		if (it->GetTickState() != NTS_Dead && it->GetTickState() != NTS_Suspended)
+		{
+			//TODO - Handle constants better. Like waaaay better.
+			it->SetConstants(Constants);
+			it->GetConstants().Merge(it->GetProperties()->SpawnScriptProps.ExternalConstants);
+			it->GetConstants().Merge(it->GetProperties()->UpdateScriptProps.ExternalConstants);
+			it->Tick(DeltaSeconds);
+		}
+
+		EffectBounds += it->GetEffectRenderer()->GetBounds();
+	}
+
+	Age += DeltaSeconds;
+}
 
 void FNiagaraEffectInstance::RenderModuleupdate()
 {
@@ -89,9 +130,8 @@ void FNiagaraEffectInstance::InitEmitters(UNiagaraEffect *InAsset)
 	check(InAsset);
 	for (int i = 0; i < InAsset->GetNumEmitters(); i++)
 	{
-		FNiagaraEmitterProperties *Props = InAsset->GetEmitterProperties(i);
+		UNiagaraEmitterProperties *Props = InAsset->GetEmitterProperties(i);
 		FNiagaraSimulation *Sim = new FNiagaraSimulation(Props, InAsset);
 		Emitters.Add(MakeShareable(Sim));
 	}
 }
-
diff --git a/Engine/Source/Runtime/Niagara/Private/NiagaraFunctionLibrary.cpp b/Engine/Source/Runtime/Niagara/Private/NiagaraFunctionLibrary.cpp
--- a/Engine/Source/Runtime/Niagara/Private/NiagaraFunctionLibrary.cpp
+++ b/Engine/Source/Runtime/Niagara/Private/NiagaraFunctionLibrary.cpp
@@ -87,12 +87,15 @@ void UNiagaraFunctionLibrary::SetUpdateScriptConstant(UNiagaraComponent* Compone
 	TArray<TSharedPtr<FNiagaraSimulation>> &Emitters = Component->GetEffectInstance()->GetEmitters();
 
 	for (TSharedPtr<FNiagaraSimulation> &Emitter : Emitters)
-	{
-		FName CurName = *Emitter->GetProperties()->Name;
-		if (CurName == EmitterName)
+	{		
+		if(UNiagaraEmitterProperties* PinnedProps = Emitter->GetProperties().Get())
 		{
-			Emitter->GetProperties()->ExternalConstants.SetOrAdd(ConstantName, Value);
-			break;
+			FName CurName = *PinnedProps->EmitterName;
+			if (CurName == EmitterName)
+			{
+				Emitter->GetProperties()->UpdateScriptProps.ExternalConstants.SetOrAdd(ConstantName, Value);
+				break;
+			}
 		}
 	}
 }
\ No newline at end of file
diff --git a/Engine/Source/Runtime/Niagara/Private/NiagaraSimulation.cpp b/Engine/Source/Runtime/Niagara/Private/NiagaraSimulation.cpp
--- a/Engine/Source/Runtime/Niagara/Private/NiagaraSimulation.cpp
+++ b/Engine/Source/Runtime/Niagara/Private/NiagaraSimulation.cpp
@@ -5,7 +5,7 @@
 #include "NiagaraEffectRenderer.h"
 #include "VectorVM.h"
 
-FNiagaraSimulation::FNiagaraSimulation(FNiagaraEmitterProperties *InProps, UNiagaraEffect *InEffect)
+FNiagaraSimulation::FNiagaraSimulation(TWeakObjectPtr<UNiagaraEmitterProperties> InProps, UNiagaraEffect* InEffect)
 : Age(0.0f)
 , Loops(0)
 , bIsEnabled(true)
@@ -19,7 +19,7 @@ FNiagaraSimulation::FNiagaraSimulation(FNiagaraEmitterProperties *InProps, UNiag
 	Init();
 }
 
-FNiagaraSimulation::FNiagaraSimulation(FNiagaraEmitterProperties *InProps, UNiagaraEffect *InEffect, ERHIFeatureLevel::Type InFeatureLevel)
+FNiagaraSimulation::FNiagaraSimulation(TWeakObjectPtr<UNiagaraEmitterProperties> InProps, UNiagaraEffect* InEffect, ERHIFeatureLevel::Type InFeatureLevel)
 	: Age(0.0f)
 	, Loops(0)
 	, bIsEnabled(true)
@@ -36,9 +36,16 @@ FNiagaraSimulation::FNiagaraSimulation(FNiagaraEmitterProperties *InProps, UNiag
 
 void FNiagaraSimulation::Init()
 {
-	if (Props && Props->UpdateScript && Props->SpawnScript )
+	Data.Reset();
+	SpawnRemainder = 0.0f;
+	Constants.Empty();
+	Age = 0.0f;
+	Loops = 0.0f;
+
+	UNiagaraEmitterProperties* PinnedProps = Props.Get();
+	if (PinnedProps && PinnedProps->UpdateScriptProps.Script && PinnedProps->SpawnScriptProps.Script)
 	{
-		if (Props->UpdateScript->Attributes.Num() == 0 || Props->SpawnScript->Attributes.Num() == 0)
+		if (PinnedProps->UpdateScriptProps.Script->Attributes.Num() == 0 || PinnedProps->SpawnScriptProps.Script->Attributes.Num() == 0)
 		{
 			Data.Reset();
 			bIsEnabled = false;
@@ -48,19 +55,20 @@ void FNiagaraSimulation::Init()
 		{
 			//Warn the user if there are any attributes used in the update script that are not initialized in the spawn script.
 			//TODO: We need some window in the effect editor and possibly the graph editor for warnings and errors.
-			for (FNiagaraVariableInfo& Attr : Props->UpdateScript->Attributes)
+			for (FNiagaraVariableInfo& Attr : PinnedProps->UpdateScriptProps.Script->Attributes)
 			{
 				int32 FoundIdx;
-				if (!Props->SpawnScript->Attributes.Find(Attr, FoundIdx))
+				if (!PinnedProps->SpawnScriptProps.Script->Attributes.Find(Attr, FoundIdx))
 				{
-					UE_LOG(LogNiagara, Warning, TEXT("Attribute %s is used in the Update script for %s but it is not initialised in the Spawn script!"), *Attr.Name.ToString(), *Props->Name);
+					UE_LOG(LogNiagara, Warning, TEXT("Attribute %s is used in the Update script for %s but it is not initialised in the Spawn script!"), *Attr.Name.ToString(), *Props->EmitterName);
 				}
 			}
-			Data.AddAttributes(Props->UpdateScript->Attributes);
-			Data.AddAttributes(Props->SpawnScript->Attributes);
+			Data.AddAttributes(PinnedProps->UpdateScriptProps.Script->Attributes);
+			Data.AddAttributes(PinnedProps->SpawnScriptProps.Script->Attributes);
 
 			Constants.Empty();
-			Constants.Merge(Props->ExternalConstants);
+			Constants.Merge(PinnedProps->UpdateScriptProps.ExternalConstants);
+			Constants.Merge(PinnedProps->SpawnScriptProps.ExternalConstants);
 		}
 	}
 	else
@@ -72,24 +80,12 @@ void FNiagaraSimulation::Init()
 	}
 }
 
-void FNiagaraSimulation::SetProperties(FNiagaraEmitterProperties *InProps)	
+void FNiagaraSimulation::SetProperties(TWeakObjectPtr<UNiagaraEmitterProperties> InProps)
 { 
 	Props = InProps; 
 	Init();
 }
 
-void FNiagaraSimulation::SetSpawnScript(UNiagaraScript* Script)
-{
-	Props->SpawnScript = Script;
-	Init();
-}
-
-void FNiagaraSimulation::SetUpdateScript(UNiagaraScript* Script)
-{
-	Props->UpdateScript = Script;
-	Init();
-}
-
 void FNiagaraSimulation::SetEnabled(bool bEnabled)
 {
 	bIsEnabled = bEnabled;
@@ -118,14 +114,16 @@ void FNiagaraSimulation::Tick(float DeltaSeconds)
 {
 	SCOPE_CYCLE_COUNTER(STAT_NiagaraTick);
 
-	if (!bIsEnabled || TickState==NTS_Suspended || TickState==NTS_Dead)
+	UNiagaraEmitterProperties* PinnedProps = Props.Get();
+	if (!PinnedProps || !bIsEnabled || TickState==NTS_Suspended || TickState==NTS_Dead)
 		return;
 
 	SimpleTimer TickTime;
 
+
 	check(Data.GetNumAttributes() > 0);
-	check(Props->SpawnScript);
-	check(Props->UpdateScript);
+	check(PinnedProps->SpawnScriptProps.Script);
+	check(PinnedProps->UpdateScriptProps.Script);
 
 	// Cache the ComponentToWorld transform.
 //	CachedComponentToWorld = Component.GetComponentToWorld();
@@ -149,7 +147,7 @@ void FNiagaraSimulation::Tick(float DeltaSeconds)
 	if (TickState==NTS_Running || TickState==NTS_Dieing)
 	{
 		SCOPE_CYCLE_COUNTER(STAT_NiagaraSimulate);
-		RunVMScript(Props->UpdateScript, EUnusedAttributeBehaviour::Copy);
+		RunVMScript(PinnedProps->UpdateScriptProps.Script, EUnusedAttributeBehaviour::Copy);
 	}
 	
 	//Init new particles with the spawn script.
@@ -158,7 +156,7 @@ void FNiagaraSimulation::Tick(float DeltaSeconds)
 		SCOPE_CYCLE_COUNTER(STAT_NiagaraSpawn);
 		Data.SetNumParticles(MaxNewParticles);
 		//For now, zero any unused attributes here. But as this is really uninitialized data we should maybe make this a more serious error.
-		RunVMScript(Props->SpawnScript, EUnusedAttributeBehaviour::Zero, OrigNumParticles, NumToSpawn);
+		RunVMScript(PinnedProps->SpawnScriptProps.Script, EUnusedAttributeBehaviour::Zero, OrigNumParticles, NumToSpawn);
 	}
 
 	// Iterate over looking for dead particles and move from the end of the list to the dead location, compacting in the process
@@ -303,7 +301,7 @@ bool FNiagaraSimulation::CheckAttriubtesForRenderer()
 			if (!Data.HasAttriubte(Attr))
 			{
 				bOk = false;
-				UE_LOG(LogNiagara, Error, TEXT("Cannot render %s because it does not define attriubte %s."), *Props->Name, *Attr.Name.ToString());
+				UE_LOG(LogNiagara, Error, TEXT("Cannot render %s because it does not define attriubte %s."), *Props->EmitterName, *Attr.Name.ToString());
 			}
 		}
 	}
@@ -348,3 +346,37 @@ void FNiagaraSimulation::SetRenderModuleType(EEmitterRenderModuleType Type, ERHI
 		CheckAttriubtesForRenderer();
 	}
 }
+
+
+//////////////////////////////////////////////////////////////////////////
+
+UNiagaraEmitterProperties::UNiagaraEmitterProperties(const FObjectInitializer& Initiilaizer)
+: Super(Initiilaizer)
+, EmitterName(TEXT("New Emitter"))
+, bIsEnabled(true)
+, SpawnRate(50)
+, Material(nullptr)
+, RenderModuleType(RMT_Sprites)
+, StartTime(0.0f)
+, EndTime(0.0f)
+, RendererProperties(nullptr)
+, NumLoops(0)
+{
+}
+
+void UNiagaraEmitterProperties::PostEditChangeProperty(struct FPropertyChangedEvent& PropertyChangedEvent)
+{
+	Super::PostEditChangeProperty(PropertyChangedEvent);
+	Init();
+}
+
+void UNiagaraEmitterProperties::PostLoad()
+{
+	Super::PostLoad();
+}
+
+void UNiagaraEmitterProperties::Serialize(FArchive& Ar)
+{
+	Super::Serialize(Ar);
+
+}
\ No newline at end of file
diff --git a/Engine/Source/Runtime/VectorVM/Private/VectorVM.cpp b/Engine/Source/Runtime/VectorVM/Private/VectorVM.cpp
--- a/Engine/Source/Runtime/VectorVM/Private/VectorVM.cpp
+++ b/Engine/Source/Runtime/VectorVM/Private/VectorVM.cpp
@@ -9,30 +9,53 @@ IMPLEMENT_MODULE(FDefaultModuleImpl, VectorVM);
 
 DEFINE_LOG_CATEGORY_STATIC(LogVectorVM, All, All);
 
-//#define VM_FORCEINLINE 
+//#define VM_FORCEINLINE
 #define VM_FORCEINLINE FORCEINLINE
 
-#define SRCOP_RRRR 0x00
-#define SRCOP_RRRC 0x01
-#define SRCOP_RRCR 0x02
-#define SRCOP_RRCC 0x03
-#define SRCOP_RCRR 0x04
-#define SRCOP_RCRC 0x05
-#define SRCOP_RCCR 0x06
-#define SRCOP_RCCC 0x07
-#define SRCOP_CRRR 0x08
-#define SRCOP_CRRC 0x09
-#define SRCOP_CRCR 0x0a
-#define SRCOP_CRCC 0x0b
-#define SRCOP_CCRR 0x0c
-#define SRCOP_CCRC 0x0d
-#define SRCOP_CCCR 0x0e
-#define SRCOP_CCCC 0x0f
-
-#define SRCOP_RRRB 0x10
-#define SRCOP_RRBR 0x20
-#define SRCOP_RRBB 0x30
+#define OP_REGISTER (0)
+#define OP0_CONST (1 << 0)
+#define OP1_CONST (1 << 1)
+#define OP2_CONST (1 << 2)
+#define OP3_CONST (1 << 3)
+#define OP0_DATAOBJ (1 << 4)
+#define OP1_DATAOBJ (1 << 5)
+#define OP2_DATAOBJ (1 << 6)
+#define OP3_DATAOBJ (1 << 7)
+
+#define SRCOP_RRRR (OP_REGISTER | OP_REGISTER | OP_REGISTER | OP_REGISTER)
+#define SRCOP_RRRC (OP_REGISTER | OP_REGISTER | OP_REGISTER | OP0_CONST)
+#define SRCOP_RRCR (OP_REGISTER | OP_REGISTER | OP1_CONST | OP_REGISTER)
+#define SRCOP_RRCC (OP_REGISTER | OP_REGISTER | OP1_CONST | OP0_CONST)
+#define SRCOP_RCRR (OP_REGISTER | OP2_CONST | OP_REGISTER | OP_REGISTER)
+#define SRCOP_RCRC (OP_REGISTER | OP2_CONST | OP_REGISTER | OP0_CONST)
+#define SRCOP_RCCR (OP_REGISTER | OP2_CONST | OP1_CONST | OP_REGISTER)
+#define SRCOP_RCCC (OP_REGISTER | OP2_CONST | OP1_CONST | OP0_CONST)
+#define SRCOP_CRRR (OP3_CONST | OP_REGISTER | OP_REGISTER | OP_REGISTER)
+#define SRCOP_CRRC (OP3_CONST | OP_REGISTER | OP_REGISTER | OP0_CONST)
+#define SRCOP_CRCR (OP3_CONST | OP_REGISTER | OP1_CONST | OP_REGISTER)
+#define SRCOP_CRCC (OP3_CONST | OP_REGISTER | OP1_CONST | OP0_CONST)
+#define SRCOP_CCRR (OP3_CONST | OP2_CONST | OP_REGISTER | OP_REGISTER)
+#define SRCOP_CCRC (OP3_CONST | OP2_CONST | OP_REGISTER | OP0_CONST)
+#define SRCOP_CCCR (OP3_CONST | OP2_CONST | OP1_CONST | OP_REGISTER)
+#define SRCOP_CCCC (OP3_CONST | OP2_CONST | OP1_CONST | OP0_CONST)
+
+#define SRCOP_RRRB (OP_REGISTER | OP_REGISTER | OP_REGISTER | OP0_DATAOBJ)
+#define SRCOP_RRBR (OP_REGISTER | OP_REGISTER | OP1_DATAOBJ | OP_REGISTER)
+#define SRCOP_RRBB (OP_REGISTER | OP_REGISTER | OP1_DATAOBJ | OP0_DATAOBJ)
+
+#define SRCOP_RRCB (OP_REGISTER | OP_REGISTER | OP1_CONST | OP0_DATAOBJ)
 
+uint8 VectorVM::CreateSrcOperandMask(VectorVM::EOperandType Type1, VectorVM::EOperandType Type2, VectorVM::EOperandType Type3, VectorVM::EOperandType Type4)
+{
+	return	(Type1 == VectorVM::ConstantOperandType ? OP0_CONST : OP_REGISTER) |
+		(Type2 == VectorVM::ConstantOperandType ? OP1_CONST : OP_REGISTER) |
+		(Type3 == VectorVM::ConstantOperandType ? OP2_CONST : OP_REGISTER) |
+		(Type4 == VectorVM::ConstantOperandType ? OP3_CONST : OP_REGISTER) |
+		(Type1 == VectorVM::DataObjConstantOperandType ? OP0_DATAOBJ : OP_REGISTER) |
+		(Type2 == VectorVM::DataObjConstantOperandType ? OP1_DATAOBJ : OP_REGISTER) |
+		(Type3 == VectorVM::DataObjConstantOperandType ? OP2_DATAOBJ : OP_REGISTER) |
+		(Type4 == VectorVM::DataObjConstantOperandType ? OP3_DATAOBJ : OP_REGISTER);
+}
 
 
 UNiagaraDataObject::UNiagaraDataObject(const FObjectInitializer& ObjectInitializer)
@@ -291,6 +314,7 @@ struct TBinaryVectorKernelData
 		switch (SrcOpTypes)
 		{
 		case SRCOP_RRRB:	VectorBinaryLoop<Kernel, FDataObjectConstantHandler, FRegisterHandler>(Context, Dst, NumVectors); break;
+		case SRCOP_RRCB:	VectorBinaryLoop<Kernel, FDataObjectConstantHandler, FConstantHandler>(Context, Dst, NumVectors); break;
 		default: check(0); break;
 		};
 	}
@@ -1093,26 +1117,6 @@ uint8 VectorVM::GetNumOpCodes()
 	return (uint8)EOp::NumOpcodes;
 }
 
-uint8 VectorVM::CreateSrcOperandMask(bool bIsOp0Constant, bool bIsOp1Constant, bool bIsOp2Constant, bool bIsOp3Constant)
-{
-	return	(bIsOp0Constant ? (1 << 0) : 0) |
-			(bIsOp1Constant ? (1 << 1) : 0) |
-			(bIsOp2Constant ? (1 << 2) : 0) |
-			(bIsOp3Constant ? (1 << 3) : 0) ;
-}
-
-uint8 VectorVM::CreateSrcOperandMask(VectorVM::EOperandType Type1, VectorVM::EOperandType Type2, VectorVM::EOperandType Type3, VectorVM::EOperandType Type4)
-{
-	return	(Type1==VectorVM::ConstantOperandType ? (1 << 0) : 0) |
-		(Type2 == VectorVM::ConstantOperandType ? (1 << 1) : 0) |
-		(Type3 == VectorVM::ConstantOperandType ? (1 << 2) : 0) |
-		(Type4 == VectorVM::ConstantOperandType ? (1 << 3) : 0) |
-		(Type1 == VectorVM::DataObjConstantOperandType ? (1 << 4) : 0) | 
-		(Type2 == VectorVM::DataObjConstantOperandType ? (1 << 5) : 0) |
-		(Type3 == VectorVM::DataObjConstantOperandType ? (1 << 6) : 0) |
-		(Type4 == VectorVM::DataObjConstantOperandType ? (1 << 7) : 0);
-}
-
 /*------------------------------------------------------------------------------
 	Automation test for the VM.
 ------------------------------------------------------------------------------*/
diff --git a/Engine/Source/Runtime/VectorVM/Public/VectorVM.h b/Engine/Source/Runtime/VectorVM/Public/VectorVM.h
--- a/Engine/Source/Runtime/VectorVM/Public/VectorVM.h
+++ b/Engine/Source/Runtime/VectorVM/Public/VectorVM.h
@@ -95,7 +95,6 @@ namespace VectorVM
 	/** Get total number of op-codes */
 	VECTORVM_API uint8 GetNumOpCodes();
 
-	VECTORVM_API uint8 CreateSrcOperandMask(bool bIsOp0Constant, bool bIsOp1Constant = false, bool bIsOp2Constant = false, bool bIsOp3Constant = false);
 	VECTORVM_API uint8 CreateSrcOperandMask(VectorVM::EOperandType Type1, VectorVM::EOperandType Type2 = VectorVM::RegisterOperandType, VectorVM::EOperandType Type3 = VectorVM::RegisterOperandType, VectorVM::EOperandType Type4 = VectorVM::RegisterOperandType);
 
 	/**
diff --git a/Engine/Source/Runtime/VectorVM/Public/VectorVMDataObject.h b/Engine/Source/Runtime/VectorVM/Public/VectorVMDataObject.h
--- a/Engine/Source/Runtime/VectorVM/Public/VectorVMDataObject.h
+++ b/Engine/Source/Runtime/VectorVM/Public/VectorVMDataObject.h
@@ -7,7 +7,7 @@
 /* Vector VM data object; encapsulates buffers, curves and other data in its derivatives
 *  for access by VectorVM kernels;
 */
-UCLASS()
+UCLASS(EditInlineNew, MinimalAPI)
 class UNiagaraDataObject : public UObject
 {
 	GENERATED_UCLASS_BODY()
@@ -18,17 +18,13 @@ public:
 
 /* Curve object; encapsulates a curve for the VectorVM
 */
-UCLASS(MinimalAPI)
+UCLASS(EditInlineNew, MinimalAPI)
 class UNiagaraCurveDataObject : public UNiagaraDataObject
 {
 	GENERATED_UCLASS_BODY()
-private:
-	UPROPERTY()
+	
+	UPROPERTY(EditAnywhere, Category="Curve")
 	class UCurveVector *CurveObj;
-public:
-	UNiagaraCurveDataObject(class UCurveVector *InCurve) : CurveObj(InCurve)
-	{
-	}
 
 	FVector4 Sample(const FVector4& InCoords) const;
 	virtual FVector4 Write(const FVector4& InCoords, const FVector4& InValue) override
@@ -43,7 +39,7 @@ public:
 
 /* Volume data object; encapsulates volumetric data for VectorVM
 */
-UCLASS(MinimalAPI, Transient)
+UCLASS(EditInlineNew, MinimalAPI, Transient)
 class UNiagaraSparseVolumeDataObject : public UNiagaraDataObject
 {
 	GENERATED_UCLASS_BODY()
