Speculative compile fix for mac.

Also a general style fix even if it doesn't work.

diff --git a/Engine/Source/Runtime/Niagara/Classes/NiagaraSimulation.h b/Engine/Source/Runtime/Niagara/Classes/NiagaraSimulation.h
--- a/Engine/Source/Runtime/Niagara/Classes/NiagaraSimulation.h
+++ b/Engine/Source/Runtime/Niagara/Classes/NiagaraSimulation.h
@@ -69,9 +69,9 @@ public:
 	}
 
 	//Begin UObject Interface
-	void PostEditChangeProperty(struct FPropertyChangedEvent& PropertyChangedEvent);
-	void PostLoad();
-	void Serialize(FArchive& Ar);
+	virtual void PostEditChangeProperty(struct FPropertyChangedEvent& PropertyChangedEvent)override;
+	virtual void PostLoad()override;
+	virtual void Serialize(FArchive& Ar)override;
 	//End UObject Interface
 
 	UPROPERTY(EditAnywhere, Category = "Emitter")
