Speculative compile fix for unity

diff --git a/Engine/Source/Runtime/Niagara/Private/NiagaraConstantSet.cpp b/Engine/Source/Runtime/Niagara/Private/NiagaraConstantSet.cpp
--- a/Engine/Source/Runtime/Niagara/Private/NiagaraConstantSet.cpp
+++ b/Engine/Source/Runtime/Niagara/Private/NiagaraConstantSet.cpp
@@ -2,6 +2,7 @@
 
 #include "NiagaraPrivate.h"
 #include "NiagaraConstantSet.h"
+#include "NiagaraSimulation.h"
 
 //////////////////////////////////////////////////////////////////////////
 // FNiagaraConstantMap
