- Various Niagara clang compile error/warning fixes

diff --git a/Engine/Source/Runtime/Niagara/Classes/NiagaraSimulation.h b/Engine/Source/Runtime/Niagara/Classes/NiagaraSimulation.h
--- a/Engine/Source/Runtime/Niagara/Classes/NiagaraSimulation.h
+++ b/Engine/Source/Runtime/Niagara/Classes/NiagaraSimulation.h
@@ -69,7 +69,9 @@ public:
 	}
 
 	//Begin UObject Interface
-	virtual void PostEditChangeProperty(struct FPropertyChangedEvent& PropertyChangedEvent)override;
+#if WITH_EDITOR
+	virtual void PostEditChangeProperty(struct FPropertyChangedEvent& PropertyChangedEvent) override;
+#endif
 	virtual void PostLoad()override;
 	virtual void Serialize(FArchive& Ar)override;
 	//End UObject Interface
diff --git a/Engine/Source/Runtime/Niagara/Private/NiagaraConstantSet.cpp b/Engine/Source/Runtime/Niagara/Private/NiagaraConstantSet.cpp
--- a/Engine/Source/Runtime/Niagara/Private/NiagaraConstantSet.cpp
+++ b/Engine/Source/Runtime/Niagara/Private/NiagaraConstantSet.cpp
@@ -393,17 +393,17 @@ void FNiagaraConstants::SetOrAdd(FName Name, UNiagaraDataObject* Value)
 
 float* FNiagaraConstants::FindScalar(FName Name)
 {
-	auto* C = ScalarConstants.FindByPredicate([&](const FNiagaraConstants_Float& C){ return C.Name == Name; });
+	auto* C = ScalarConstants.FindByPredicate([&](const FNiagaraConstants_Float& Constant){ return Constant.Name == Name; });
 	return C ? &C->Value : nullptr;
 }
 FVector4* FNiagaraConstants::FindVector(FName Name)
 {
-	auto* C = VectorConstants.FindByPredicate([&](const FNiagaraConstants_Vector& C){ return C.Name == Name; });
+	auto* C = VectorConstants.FindByPredicate([&](const FNiagaraConstants_Vector& Constant){ return Constant.Name == Name; });
 	return C ? &C->Value : nullptr;
 }
 FMatrix* FNiagaraConstants::FindMatrix(FName Name)
 {
-	auto* C = MatrixConstants.FindByPredicate([&](const FNiagaraConstants_Matrix& C){ return C.Name == Name; });
+	auto* C = MatrixConstants.FindByPredicate([&](const FNiagaraConstants_Matrix& Constant){ return Constant.Name == Name; });
 	return C ? &C->Value : nullptr;
 }
 UNiagaraDataObject* FNiagaraConstants::FindDataObj(FName Name)
diff --git a/Engine/Source/Runtime/Niagara/Private/NiagaraSimulation.cpp b/Engine/Source/Runtime/Niagara/Private/NiagaraSimulation.cpp
--- a/Engine/Source/Runtime/Niagara/Private/NiagaraSimulation.cpp
+++ b/Engine/Source/Runtime/Niagara/Private/NiagaraSimulation.cpp
@@ -364,11 +364,13 @@ UNiagaraEmitterProperties::UNiagaraEmitterProperties(const FObjectInitializer& I
 {
 }
 
+#if WITH_EDITOR
 void UNiagaraEmitterProperties::PostEditChangeProperty(struct FPropertyChangedEvent& PropertyChangedEvent)
 {
 	Super::PostEditChangeProperty(PropertyChangedEvent);
 	Init();
 }
+#endif
 
 void UNiagaraEmitterProperties::PostLoad()
 {
