- Missed a shadow variable in last checkin

diff --git a/Engine/Source/Runtime/Niagara/Private/NiagaraConstantSet.cpp b/Engine/Source/Runtime/Niagara/Private/NiagaraConstantSet.cpp
--- a/Engine/Source/Runtime/Niagara/Private/NiagaraConstantSet.cpp
+++ b/Engine/Source/Runtime/Niagara/Private/NiagaraConstantSet.cpp
@@ -408,7 +408,7 @@ FMatrix* FNiagaraConstants::FindMatrix(FName Name)
 }
 UNiagaraDataObject* FNiagaraConstants::FindDataObj(FName Name)
 {
-	auto* C = DataObjectConstants.FindByPredicate([&](const FNiagaraConstants_DataObject& C){ return C.Name == Name; });
+	auto* C = DataObjectConstants.FindByPredicate([&](const FNiagaraConstants_DataObject& Constant){ return Constant.Name == Name; });
 	return C ? C->Value : nullptr;
 }
 
