Fixing breakages to existing data introduced by the recent dev UI and data object changes.

diff --git a/Engine/Source/Runtime/Core/Public/UObject/ObjectVersion.h b/Engine/Source/Runtime/Core/Public/UObject/ObjectVersion.h
--- a/Engine/Source/Runtime/Core/Public/UObject/ObjectVersion.h
+++ b/Engine/Source/Runtime/Core/Public/UObject/ObjectVersion.h
@@ -532,6 +532,8 @@ enum EUnrealEngineObjectUE4Version
 	VER_UE4_DEPRECATE_UMG_STYLE_OVERRIDES,
 	// Shadowmap penumbra size stored
 	VER_UE4_STATIC_SHADOWMAP_PENUMBRA_SIZE,
+	// Fix BC on Niagara effects from the data object and dev UI changes.
+	VER_UE4_NIAGARA_DATA_OBJECT_DEV_UI_FIX,
 
 	// -----<new versions can be added before this line>-------------------------------------------------
 	// - this needs to be the last line (see note below)
diff --git a/Engine/Source/Runtime/Niagara/Classes/NiagaraEffect.h b/Engine/Source/Runtime/Niagara/Classes/NiagaraEffect.h
--- a/Engine/Source/Runtime/Niagara/Classes/NiagaraEffect.h
+++ b/Engine/Source/Runtime/Niagara/Classes/NiagaraEffect.h
@@ -43,7 +43,12 @@ public:
 
 	void CreateEffectRendererProps(TSharedPtr<FNiagaraSimulation> Sim);
 
+	//Begin UObject Interface
+	virtual void Serialize(FArchive& Ar);
+	//End UObject Interface
 private:
+	UPROPERTY()
+	TArray<FDeprecatedNiagaraEmitterProperties>EmitterPropsSerialized_DEPRECATED;
 
 	UPROPERTY()
 	TArray<UNiagaraEmitterProperties*> EmitterProps;
diff --git a/Engine/Source/Runtime/Niagara/Classes/NiagaraSimulation.h b/Engine/Source/Runtime/Niagara/Classes/NiagaraSimulation.h
--- a/Engine/Source/Runtime/Niagara/Classes/NiagaraSimulation.h
+++ b/Engine/Source/Runtime/Niagara/Classes/NiagaraSimulation.h
@@ -52,9 +52,42 @@ struct FNiagaraEmitterScriptProperties
 	}
 };
 
+//This struct now only exists for backwards compatibility and should be removed once effects are updated.
+USTRUCT()
+struct FDeprecatedNiagaraEmitterProperties
+{
+	GENERATED_USTRUCT_BODY()
+public:
+	UPROPERTY()
+	FString Name;
+	UPROPERTY()
+	bool bIsEnabled;
+	UPROPERTY()
+	float SpawnRate;
+	UPROPERTY()
+	UNiagaraScript *UpdateScript;
+	UPROPERTY()
+	UNiagaraScript *SpawnScript;
+	UPROPERTY()
+	UMaterial *Material;
+	UPROPERTY()
+	TEnumAsByte<EEmitterRenderModuleType> RenderModuleType;
+	UPROPERTY()
+	float StartTime;
+	UPROPERTY()
+	float EndTime;
+	UPROPERTY()
+	class UNiagaraEffectRendererProperties *RendererProperties;
+	UPROPERTY()
+	FNiagaraConstantMap ExternalConstants;		// these are the update script constants from the effect editor; will be added to the emitter's constant map
+	UPROPERTY()
+	FNiagaraConstantMap ExternalSpawnConstants;		// these are the spawn script constants from the effect editor; will be added to the emitter's constant map
+	UPROPERTY()
+	int32 NumLoops;
+};
 
 /** 
- *	FniagaraEmitterProperties stores the attributes of an FNiagaraSimulation
+ *	UNiagaraEmitterProperties stores the attributes of an FNiagaraSimulation
  *	that need to be serialized and are used for its initialization 
  */
 UCLASS(MinimalAPI)
@@ -68,12 +101,12 @@ public:
 		UpdateScriptProps.Init(this);
 	}
 
+	void InitFromOldStruct(FDeprecatedNiagaraEmitterProperties& OldStruct);
+
 	//Begin UObject Interface
 #if WITH_EDITOR
 	virtual void PostEditChangeProperty(struct FPropertyChangedEvent& PropertyChangedEvent) override;
 #endif
-	virtual void PostLoad()override;
-	virtual void Serialize(FArchive& Ar)override;
 	//End UObject Interface
 
 	UPROPERTY(EditAnywhere, Category = "Emitter")
diff --git a/Engine/Source/Runtime/Niagara/NiagaraConstantSet.h b/Engine/Source/Runtime/Niagara/NiagaraConstantSet.h
--- a/Engine/Source/Runtime/Niagara/NiagaraConstantSet.h
+++ b/Engine/Source/Runtime/Niagara/NiagaraConstantSet.h
@@ -8,9 +8,11 @@
 
 struct FNiagaraConstants;
 
-/** Runtime storage for emitter overridden constants. */
+/** Runtime storage for emitter overridden constants. Only now a USTRUCT still to support BC */
+USTRUCT()
 struct FNiagaraConstantMap
 {
+	GENERATED_USTRUCT_BODY()
 private:
 	TMap<FNiagaraVariableInfo, float> ScalarConstants;
 	TMap<FNiagaraVariableInfo, FVector4> VectorConstants;
@@ -50,7 +52,22 @@ public:
 	void Merge(FNiagaraConstants &InConstants);
 	void Merge(FNiagaraConstantMap &InMap);
 
+	const TMap<FNiagaraVariableInfo, float> &GetScalarConstants()	const;
+	const TMap<FNiagaraVariableInfo, FVector4> &GetVectorConstants()	const;
+	const TMap<FNiagaraVariableInfo, FMatrix> &GetMatrixConstants()	const;
 	const TMap<FNiagaraVariableInfo, class UNiagaraDataObject*> &GetDataConstants()	const;
+
+
+	virtual bool Serialize(FArchive &Ar)
+	{
+		// TODO: can't serialize the data object constants at the moment; need to figure that out
+		Ar << ScalarConstants << VectorConstants << MatrixConstants;
+		if (Ar.UE4Ver() >= VER_UE4_NIAGARA_DATA_OBJECT_DEV_UI_FIX)
+		{
+			Ar << DataConstants;
+		}
+		return true;
+	}
 };
 
 USTRUCT()
@@ -197,6 +214,12 @@ public:
 
 	NIAGARA_API void Init(class UNiagaraEmitterProperties* EmitterProps, struct FNiagaraEmitterScriptProperties* ScriptProps);
 
+	/**
+	Copies constants in from the script as usual but allows overrides from old data in an FNiagaraConstantMap.
+	Only used for BC. DO NOT USE.
+	*/
+	NIAGARA_API void InitFromOldMap(FNiagaraConstantMap& OldMap);
+
 	/** Fills the entire constants set into the constant table. */
 	void AppendToConstantsTable(TArray<FVector4>& ConstantsTable)const;
 
diff --git a/Engine/Source/Runtime/Niagara/NiagaraScriptConstantData.h b/Engine/Source/Runtime/Niagara/NiagaraScriptConstantData.h
--- a/Engine/Source/Runtime/Niagara/NiagaraScriptConstantData.h
+++ b/Engine/Source/Runtime/Niagara/NiagaraScriptConstantData.h
@@ -14,6 +14,8 @@ struct FNiagaraScriptConstantData
 	/** The set of external constants for this script. */
 	UPROPERTY()
 	FDeprecatedNiagaraConstants ExternalConstants_DEPRECATED;
+	UPROPERTY()
+	FDeprecatedNiagaraConstants InternalConstants_DEPRECATED;
 
 	UPROPERTY()
 	FNiagaraConstants ExposedConstants;
diff --git a/Engine/Source/Runtime/Niagara/Private/NiagaraConstantSet.cpp b/Engine/Source/Runtime/Niagara/Private/NiagaraConstantSet.cpp
--- a/Engine/Source/Runtime/Niagara/Private/NiagaraConstantSet.cpp
+++ b/Engine/Source/Runtime/Niagara/Private/NiagaraConstantSet.cpp
@@ -7,6 +7,16 @@
 //////////////////////////////////////////////////////////////////////////
 // FNiagaraConstantMap
 
+template<>
+struct TStructOpsTypeTraits<FNiagaraConstantMap> : public TStructOpsTypeTraitsBase
+{
+	enum
+	{
+		WithSerializer = true
+	};
+};
+IMPLEMENT_STRUCT(NiagaraConstantMap);
+
 void FNiagaraConstantMap::Empty()
 {
 	ScalarConstants.Empty();
@@ -137,6 +147,21 @@ void FNiagaraConstantMap::Merge(FNiagaraConstantMap &InMap)
 	}
 }
 
+const TMap<FNiagaraVariableInfo, float> &FNiagaraConstantMap::GetScalarConstants()	const
+{
+	return ScalarConstants;
+}
+
+const TMap<FNiagaraVariableInfo, FVector4> &FNiagaraConstantMap::GetVectorConstants()	const
+{
+	return VectorConstants;
+}
+
+const TMap<FNiagaraVariableInfo, FMatrix> &FNiagaraConstantMap::GetMatrixConstants()	const
+{
+	return MatrixConstants;
+}
+
 const TMap<FNiagaraVariableInfo, class UNiagaraDataObject*> &FNiagaraConstantMap::GetDataConstants()	const
 {
 	return DataConstants;
@@ -238,6 +263,35 @@ void FNiagaraConstants::Init(UNiagaraEmitterProperties* EmitterProps, FNiagaraEm
 	}
 }
 
+/** 
+Copies constants in from the script as usual but allows overrides from old data in an FNiagaraConstantMap.
+Only used for BC. DO NOT USE.
+*/
+void FNiagaraConstants::InitFromOldMap(FNiagaraConstantMap& OldMap)
+{
+	const TMap<FNiagaraVariableInfo, float>& OldScalarConsts = OldMap.GetScalarConstants();
+	const TMap<FNiagaraVariableInfo, FVector4>& OldVectorConsts = OldMap.GetVectorConstants();
+	const TMap<FNiagaraVariableInfo, FMatrix>& OldMatrixConsts = OldMap.GetMatrixConstants();
+	const TMap<FNiagaraVariableInfo, UNiagaraDataObject*>& OldDataObjectConstants = OldMap.GetDataConstants();
+
+	for (const TPair<FNiagaraVariableInfo, float> Pair : OldScalarConsts)
+	{
+		ScalarConstants.Add(FNiagaraConstants_Float(Pair.Key.Name, Pair.Value));
+	}
+	for (const TPair<FNiagaraVariableInfo, FVector4>& Pair : OldVectorConsts)
+	{
+		VectorConstants.Add(FNiagaraConstants_Vector(Pair.Key.Name, Pair.Value));
+	}
+	for (const TPair<FNiagaraVariableInfo, FMatrix>& Pair : OldMatrixConsts)
+	{
+		MatrixConstants.Add(FNiagaraConstants_Matrix(Pair.Key.Name, Pair.Value));
+	}
+	for (const TPair<FNiagaraVariableInfo, UNiagaraDataObject*>& Pair : OldDataObjectConstants)
+	{
+		DataObjectConstants.Add(FNiagaraConstants_DataObject(Pair.Key.Name, Pair.Value));
+	}
+}
+
 void FNiagaraConstants::AppendToConstantsTable(TArray<FVector4>& ConstantsTable)const
 {
 	int32 Idx = ConstantsTable.Num();
diff --git a/Engine/Source/Runtime/Niagara/Private/NiagaraEffect.cpp b/Engine/Source/Runtime/Niagara/Private/NiagaraEffect.cpp
--- a/Engine/Source/Runtime/Niagara/Private/NiagaraEffect.cpp
+++ b/Engine/Source/Runtime/Niagara/Private/NiagaraEffect.cpp
@@ -42,6 +42,24 @@ void UNiagaraEffect::CreateEffectRendererProps(TSharedPtr<FNiagaraSimulation> Si
 	}
 }
 
+void UNiagaraEffect::Serialize(FArchive& Ar)
+{
+	Super::Serialize(Ar);
+
+	if (Ar.IsLoading() && Ar.UE4Ver() < VER_UE4_NIAGARA_DATA_OBJECT_DEV_UI_FIX)
+	{
+		int32 NumEmitters = EmitterPropsSerialized_DEPRECATED.Num();
+		for (int32 i = 0; i < NumEmitters; ++i)
+		{
+			UNiagaraEmitterProperties* NewProps = NewObject<UNiagaraEmitterProperties>(this);
+			NewProps->InitFromOldStruct(EmitterPropsSerialized_DEPRECATED[i]);
+			AddEmitterProperties(NewProps);
+		}
+	}
+}
+
+//////////////////////////////////////////////////////////////////////////
+
 TSharedPtr<FNiagaraSimulation> FNiagaraEffectInstance::AddEmitter(UNiagaraEmitterProperties *Properties)
 {
 	FNiagaraSimulation *SimPtr = new FNiagaraSimulation(Properties, Effect);
diff --git a/Engine/Source/Runtime/Niagara/Private/NiagaraSimulation.cpp b/Engine/Source/Runtime/Niagara/Private/NiagaraSimulation.cpp
--- a/Engine/Source/Runtime/Niagara/Private/NiagaraSimulation.cpp
+++ b/Engine/Source/Runtime/Niagara/Private/NiagaraSimulation.cpp
@@ -372,13 +372,20 @@ void UNiagaraEmitterProperties::PostEditChangeProperty(struct FPropertyChangedEv
 }
 #endif
 
-void UNiagaraEmitterProperties::PostLoad()
+void UNiagaraEmitterProperties::InitFromOldStruct(FDeprecatedNiagaraEmitterProperties& OldStruct)
 {
-	Super::PostLoad();
-}
-
-void UNiagaraEmitterProperties::Serialize(FArchive& Ar)
-{
-	Super::Serialize(Ar);
-
+	EmitterName = OldStruct.Name;
+	bIsEnabled = OldStruct.bIsEnabled;
+	SpawnRate = OldStruct.SpawnRate;
+	UpdateScriptProps.Script = OldStruct.UpdateScript;
+	SpawnScriptProps.Script = OldStruct.SpawnScript;
+	Material = OldStruct.Material;
+	RenderModuleType = OldStruct.RenderModuleType;
+	StartTime = OldStruct.StartTime;
+	EndTime = OldStruct.EndTime;
+	RendererProperties = OldStruct.RendererProperties;
+
+	UpdateScriptProps.ExternalConstants.InitFromOldMap(OldStruct.ExternalConstants);
+	SpawnScriptProps.ExternalConstants.InitFromOldMap(OldStruct.ExternalSpawnConstants);
+	NumLoops = OldStruct.NumLoops;
 }
\ No newline at end of file
