Additional Niagara BC fix.
Was getting warnings about duplicate property name. Renamed for now. Will rename back once data is updated.

diff --git a/Engine/Source/Editor/UnrealEd/Private/Niagara/NiagaraScriptSource.cpp b/Engine/Source/Editor/UnrealEd/Private/Niagara/NiagaraScriptSource.cpp
--- a/Engine/Source/Editor/UnrealEd/Private/Niagara/NiagaraScriptSource.cpp
+++ b/Engine/Source/Editor/UnrealEd/Private/Niagara/NiagaraScriptSource.cpp
@@ -106,7 +106,7 @@ void UNiagaraScriptSource::Compile()
 	FNiagaraEditorModule& NiagaraEditorModule = FModuleManager::Get().LoadModuleChecked<FNiagaraEditorModule>(TEXT("NiagaraEditor"));
 	NiagaraEditorModule.CompileScript(ScriptOwner);
 
-	FNiagaraConstants& ExternalConsts = ScriptOwner->ConstantData.ExposedConstants;
+	FNiagaraConstants& ExternalConsts = ScriptOwner->ConstantData.GetExternalConstants();
 
 	//Build the constant list. 
 	//This is mainly just jumping through some hoops for the custom UI. Should be removed and have the UI just read directly from the constants stored in the UScript.
diff --git a/Engine/Source/Runtime/Niagara/NiagaraScriptConstantData.h b/Engine/Source/Runtime/Niagara/NiagaraScriptConstantData.h
--- a/Engine/Source/Runtime/Niagara/NiagaraScriptConstantData.h
+++ b/Engine/Source/Runtime/Niagara/NiagaraScriptConstantData.h
@@ -17,42 +17,45 @@ struct FNiagaraScriptConstantData
 	UPROPERTY()
 	FDeprecatedNiagaraConstants InternalConstants_DEPRECATED;
 
+	/** Constants driven by the system. Named New for BC reasons. Once all data is updated beyond VER_UE4_NIAGARA_DATA_OBJECT_DEV_UI_FIX. Get rid of the deprecated consts and rename the New. */
 	UPROPERTY()
-	FNiagaraConstants ExposedConstants;
+	FNiagaraConstants ExternalConstantsNew;
 
-	/** The set of internal constants for this script. */
+	/** The set of internal constants for this script. Named New for BC reasons. Once all data is updated beyond VER_UE4_NIAGARA_DATA_OBJECT_DEV_UI_FIX. Get rid of the deprecated consts and rename the New. */
 	UPROPERTY()
-	FNiagaraConstants InternalConstants;
+	FNiagaraConstants InternalConstantsNew;
 
 	void Empty()
 	{
-		ExposedConstants.Empty();
-		InternalConstants.Empty();
+		ExternalConstantsNew.Empty();
+		InternalConstantsNew.Empty();
 	}
 
+	FNiagaraConstants& GetExternalConstants(){ return ExternalConstantsNew; }
+
 	/** Fill a constants table ready for use in an update or spawn. */
 	void FillConstantTable(const FNiagaraConstantMap& ExternalConstantsMap, TArray<FVector4>& OutConstantTable, TArray<UNiagaraDataObject *> &DataObjTable) const
 	{
 		//First up in the table comes the External constants in scalar->vector->matrix order.
 		//Only fills the constants actually used by the script.
 		OutConstantTable.Empty();
-		ExposedConstants.AppendToConstantsTable(OutConstantTable, ExternalConstantsMap);
-		ExposedConstants.AppendExternalBufferConstants(DataObjTable, ExternalConstantsMap);
+		ExternalConstantsNew.AppendToConstantsTable(OutConstantTable, ExternalConstantsMap);
+		ExternalConstantsNew.AppendExternalBufferConstants(DataObjTable, ExternalConstantsMap);
 		//Next up add all the internal constants from the script.
-		InternalConstants.AppendToConstantsTable(OutConstantTable);
-		InternalConstants.AppendBufferConstants(DataObjTable);
+		InternalConstantsNew.AppendToConstantsTable(OutConstantTable);
+		InternalConstantsNew.AppendBufferConstants(DataObjTable);
 	}
 
 	template<typename T>
 	void SetOrAddInternal(const FNiagaraVariableInfo& Constant, const T& Value)
 	{
-		InternalConstants.SetOrAdd(Constant, Value);
+		InternalConstantsNew.SetOrAdd(Constant, Value);
 	}
 
 	template<typename T>
 	void SetOrAddExternal(const FNiagaraVariableInfo& Constant, const T& Value)
 	{
-		ExposedConstants.SetOrAdd(Constant, Value);
+		ExternalConstantsNew.SetOrAdd(Constant, Value);
 	}
 
 	/**
@@ -61,11 +64,11 @@ struct FNiagaraScriptConstantData
 	*/
 	void GetTableIndex(const FNiagaraVariableInfo& Constant, bool bInternal, int32& OutConstantIdx, int32& OutComponentIndex, ENiagaraDataType& OutType)const
 	{
-		int32 Base = bInternal ? ExposedConstants.GetTableSize() : 0;
+		int32 Base = bInternal ? ExternalConstantsNew.GetTableSize() : 0;
 		int32 ConstIdx = INDEX_NONE;
 		int32 CompIdx = INDEX_NONE;
 		ENiagaraDataType Type;
-		const FNiagaraConstants& Consts = bInternal ? InternalConstants : ExposedConstants;
+		const FNiagaraConstants& Consts = bInternal ? InternalConstantsNew : ExternalConstantsNew;
 
 		ConstIdx = Consts.GetAbsoluteIndex_Scalar(Constant);
 		Type = ENiagaraDataType::Scalar;
diff --git a/Engine/Source/Runtime/Niagara/Private/NiagaraConstantSet.cpp b/Engine/Source/Runtime/Niagara/Private/NiagaraConstantSet.cpp
--- a/Engine/Source/Runtime/Niagara/Private/NiagaraConstantSet.cpp
+++ b/Engine/Source/Runtime/Niagara/Private/NiagaraConstantSet.cpp
@@ -207,7 +207,7 @@ void FNiagaraConstants::Init(UNiagaraEmitterProperties* EmitterProps, FNiagaraEm
 	check(EmitterProps);
 
 	//Copy the script constants, retaining any altered values.
-	FNiagaraConstants& ScriptConstants = ScriptProps->Script->ConstantData.ExposedConstants;
+	const FNiagaraConstants& ScriptConstants = ScriptProps->Script->ConstantData.GetExternalConstants();
 
 	TArray<FNiagaraConstants_Float> NewScalarConstants = ScriptConstants.ScalarConstants;
 	for (auto& Override : ScalarConstants)
@@ -240,10 +240,10 @@ void FNiagaraConstants::Init(UNiagaraEmitterProperties* EmitterProps, FNiagaraEm
 	}
 	MatrixConstants = NewMatrixConstants;
 
-	TArray<FNiagaraConstants_DataObject>& ScriptDataObjectConstants = ScriptConstants.DataObjectConstants;
+	const TArray<FNiagaraConstants_DataObject>& ScriptDataObjectConstants = ScriptConstants.DataObjectConstants;
 	TArray<FNiagaraConstants_DataObject> OldDataObjectConstants = DataObjectConstants;
 	DataObjectConstants.Empty();
-	for (FNiagaraConstants_DataObject& ScriptConst : ScriptDataObjectConstants)
+	for (const FNiagaraConstants_DataObject& ScriptConst : ScriptDataObjectConstants)
 	{
 		int32 OverrideIdx = OldDataObjectConstants.IndexOfByPredicate([&](const FNiagaraConstants_DataObject& C){ return ScriptConst.Name == C.Name && ScriptConst.Value && C.Value &&ScriptConst.Value->GetClass() == C.Value->GetClass(); });
 		
