** PRIMARY FEATURE **

diff --git a/Engine/Source/Runtime/Niagara/Classes/NiagaraSimulation.h b/Engine/Source/Runtime/Niagara/Classes/NiagaraSimulation.h
--- a/Engine/Source/Runtime/Niagara/Classes/NiagaraSimulation.h
+++ b/Engine/Source/Runtime/Niagara/Classes/NiagaraSimulation.h
@@ -7,157 +7,39 @@ NiagaraSimulation.h: Niagara emitter simulation class
 
 #include "NiagaraScript.h"
 #include "NiagaraComponent.h"
-#include "NiagaraSimulation.generated.h"
+#include "NiagaraEvents.h"
+#include "NiagaraDataSet.h"
+#include "NiagaraEmitterProperties.h"
+#include "Private/NiagaraSimulationDebugger.h"
 
 DECLARE_CYCLE_STAT(TEXT("Tick"), STAT_NiagaraTick, STATGROUP_Niagara);
 DECLARE_CYCLE_STAT(TEXT("Simulate"), STAT_NiagaraSimulate, STATGROUP_Niagara);
 DECLARE_CYCLE_STAT(TEXT("Spawn"), STAT_NiagaraSpawn, STATGROUP_Niagara);
 DECLARE_CYCLE_STAT(TEXT("Kill"), STAT_NiagaraKill, STATGROUP_Niagara);
 
-
-class NiagaraEffectRenderer;
-
-UENUM()
-enum EEmitterRenderModuleType
-{
-	RMT_None = 0,
-	RMT_Sprites,
-	RMT_Ribbon,
-	RMT_Trails,
-	RMT_Meshes
-};
-
-
-USTRUCT()
-struct FNiagaraEmitterScriptProperties
-{
-	GENERATED_BODY()
-
-	UPROPERTY(EditAnywhere, Category = "Script")
-	UNiagaraScript *Script;
-
-	UPROPERTY(EditAnywhere, Category = "Script", meta = (ShowOnlyInnerProperties))
-	FNiagaraConstants ExternalConstants;
-
-	void Init(UNiagaraEmitterProperties* EmitterProps)
-	{
-		if (Script)
-		{
-			ExternalConstants.Init(EmitterProps, this);
-		}
-		else
-		{
-			ExternalConstants.Empty();
-		}
-	}
-};
-
-//This struct now only exists for backwards compatibility and should be removed once effects are updated.
-USTRUCT()
-struct FDeprecatedNiagaraEmitterProperties
-{
-	GENERATED_USTRUCT_BODY()
-public:
-	UPROPERTY()
-	FString Name;
-	UPROPERTY()
-	bool bIsEnabled;
-	UPROPERTY()
-	float SpawnRate;
-	UPROPERTY()
-	UNiagaraScript *UpdateScript;
-	UPROPERTY()
-	UNiagaraScript *SpawnScript;
-	UPROPERTY()
-	UMaterial *Material;
-	UPROPERTY()
-	TEnumAsByte<EEmitterRenderModuleType> RenderModuleType;
-	UPROPERTY()
-	float StartTime;
-	UPROPERTY()
-	float EndTime;
-	UPROPERTY()
-	class UNiagaraEffectRendererProperties *RendererProperties;
-	UPROPERTY()
-	FNiagaraConstantMap ExternalConstants;		// these are the update script constants from the effect editor; will be added to the emitter's constant map
-	UPROPERTY()
-	FNiagaraConstantMap ExternalSpawnConstants;		// these are the spawn script constants from the effect editor; will be added to the emitter's constant map
-	UPROPERTY()
-	int32 NumLoops;
-};
-
-/** 
- *	UNiagaraEmitterProperties stores the attributes of an FNiagaraSimulation
- *	that need to be serialized and are used for its initialization 
- */
-UCLASS(MinimalAPI)
-class UNiagaraEmitterProperties : public UObject
-{
-	GENERATED_UCLASS_BODY()
-public:
-	void Init()
-	{
-		SpawnScriptProps.Init(this);
-		UpdateScriptProps.Init(this);
-	}
-
-	void InitFromOldStruct(FDeprecatedNiagaraEmitterProperties& OldStruct);
-
-	//Begin UObject Interface
-#if WITH_EDITOR
-	virtual void PostEditChangeProperty(struct FPropertyChangedEvent& PropertyChangedEvent) override;
-#endif
-	//End UObject Interface
-
-	UPROPERTY(EditAnywhere, Category = "Emitter")
-	FString EmitterName;
-	UPROPERTY(EditAnywhere, Category = "Emitter")
-	bool bIsEnabled;
-	UPROPERTY(EditAnywhere, Category = "Emitter")
-	float SpawnRate;
-	UPROPERTY(EditAnywhere, Category = "Emitter")
-	UMaterial *Material;
-	UPROPERTY(EditAnywhere, Category = "Emitter")
-	TEnumAsByte<EEmitterRenderModuleType> RenderModuleType;
-
-	UPROPERTY(EditAnywhere, Category = "Emitter")
-	float StartTime;
-	UPROPERTY(EditAnywhere, Category = "Emitter")
-	float EndTime;
-
-	UPROPERTY()
-	class UNiagaraEffectRendererProperties *RendererProperties;
-
-	UPROPERTY(EditAnywhere, Category = "Emitter")
-	int32 NumLoops;
-
-	UPROPERTY(EditAnywhere, Category = "Emitter", meta = (ShowOnlyInnerProperties))
-	FNiagaraEmitterScriptProperties UpdateScriptProps;
-
-	UPROPERTY(EditAnywhere, Category = "Emitter", meta = (ShowOnlyInnerProperties))
-	FNiagaraEmitterScriptProperties SpawnScriptProps;
-};
-
 /**
 * A Niagara particle simulation.
 */
 struct FNiagaraSimulation
 {
 public:
-	explicit FNiagaraSimulation(TWeakObjectPtr<UNiagaraEmitterProperties> InProps, UNiagaraEffect* Effect);
-	FNiagaraSimulation(TWeakObjectPtr<UNiagaraEmitterProperties> Props, UNiagaraEffect* Effect, ERHIFeatureLevel::Type InFeatureLevel);
+	explicit FNiagaraSimulation(TWeakObjectPtr<UNiagaraEmitterProperties> InProps, FNiagaraEffectInstance* InParentEffectInstance);
+	FNiagaraSimulation(TWeakObjectPtr<UNiagaraEmitterProperties> Props, FNiagaraEffectInstance* InParentEffectInstance, ERHIFeatureLevel::Type InFeatureLevel);
 	virtual ~FNiagaraSimulation()
 	{}
 
 	NIAGARA_API void Init();
 
+	/** Called after all emitters in an effect have been initialized, allows emitters to access information from one another. */
+	void PostInit();
+
+	void PreTick();
 	void Tick(float DeltaSeconds);
 
 
 	FBox GetBounds() const { return CachedBounds; }
 
-
-	FNiagaraEmitterParticleData &GetData()	{ return Data; }
+	FNiagaraDataSet &GetData()	{ return Data; }
 
 	void SetConstants(const FNiagaraConstantMap &InMap)	{ Constants = InMap; }
 	FNiagaraConstantMap &GetConstants()	{ return Constants; }
@@ -168,12 +50,12 @@ public:
 
 	void NIAGARA_API SetRenderModuleType(EEmitterRenderModuleType Type, ERHIFeatureLevel::Type FeatureLevel);
 
-	int32 GetNumParticles()	{ return Data.GetNumParticles(); }
+	int32 GetNumParticles()	{ return Data.GetNumInstances(); }
 
 	TWeakObjectPtr<UNiagaraEmitterProperties> GetProperties()	{ return Props; }
 	void SetProperties(TWeakObjectPtr<UNiagaraEmitterProperties> InProps);
 
-	UNiagaraEffect *GetParentEffect()	{ return ParentEffect; }
+	FNiagaraEffectInstance* GetParentEffectInstance()	{ return ParentEffectInstance; }
 
 	float NIAGARA_API GetTotalCPUTime();
 	int	NIAGARA_API GetTotalBytesUsed();
@@ -191,6 +73,11 @@ public:
 
 	ENiagaraTickState NIAGARA_API GetTickState()	{ return TickState; }
 	void NIAGARA_API SetTickState(ENiagaraTickState InState)	{ TickState = InState; }
+
+	FNiagaraDataSet* GetDataSet(FNiagaraDataSetID SetID);
+
+	void SpawnBurst(uint32 Count) { SpawnRemainder += Count; }
+
 private:
 	TWeakObjectPtr<UNiagaraEmitterProperties> Props;		// points to an entry in the array of FNiagaraProperties stored in the EffectInstance (itself pointing to the effect's properties)
 
@@ -208,7 +95,7 @@ private:
 	/** Local constant set. */
 	FNiagaraConstantMap Constants;
 	/** particle simulation data */
-	FNiagaraEmitterParticleData Data;
+	FNiagaraDataSet Data;
 	/** Keep partial particle spawns from last frame */
 	float SpawnRemainder;
 	/** The cached ComponentToWorld transform. */
@@ -217,8 +104,16 @@ private:
 	FBox CachedBounds;
 
 	NiagaraEffectRenderer *EffectRenderer;
-	UNiagaraEffect *ParentEffect;
+	FNiagaraEffectInstance *ParentEffectInstance;
 
+	TArray<FNiagaraDataSet> DataSets;
+	TMap<FNiagaraDataSetID, FNiagaraDataSet*> DataSetMap;
+	
+	FNiagaraSpawnEventGenerator SpawnEventGenerator;
+	bool bGenerateSpawnEvents;
+	FNiagaraDeathEventGenerator DeathEventGenerator;
+	bool bGenerateDeathEvents;
+	
 	/** Calc number to spawn */
 	int32 CalcNumToSpawn(float DeltaSeconds)
 	{
@@ -234,9 +129,9 @@ private:
 	}
 	
 	/** Runs a script in the VM over a specific range of particles. */
-	void RunVMScript(UNiagaraScript* Script, EUnusedAttributeBehaviour UnusedAttribBehaviour);
-	void RunVMScript(UNiagaraScript* Script, EUnusedAttributeBehaviour UnusedAttribBehaviour, uint32 StartParticle);
-	void RunVMScript(UNiagaraScript* Script, EUnusedAttributeBehaviour UnusedAttribBehaviour, uint32 StartParticle, uint32 NumParticles);
+	void RunVMScript(FNiagaraEmitterScriptProperties& ScriptProps, EUnusedAttributeBehaviour UnusedAttribBehaviour);
+	void RunVMScript(FNiagaraEmitterScriptProperties& ScriptProps, EUnusedAttributeBehaviour UnusedAttribBehaviour, uint32 StartParticle);
+	void RunVMScript(FNiagaraEmitterScriptProperties& ScriptProps, EUnusedAttributeBehaviour UnusedAttribBehaviour, uint32 StartParticle, uint32 NumParticles);
 
 	/** Util to move a particle */
 	void MoveParticleToIndex(int32 SrcIndex, int32 DestIndex)
@@ -244,13 +139,38 @@ private:
 		FVector4 *SrcPtr = Data.GetCurrentBuffer() + SrcIndex;
 		FVector4 *DestPtr = Data.GetCurrentBuffer() + DestIndex;
 
-		for (int32 AttrIndex = 0; AttrIndex < Data.GetNumAttributes(); AttrIndex++)
+		for (int32 AttrIndex = 0; AttrIndex < Data.GetNumVariables(); AttrIndex++)
 		{
 			*DestPtr = *SrcPtr;
-			DestPtr += Data.GetParticleAllocation();
-			SrcPtr += Data.GetParticleAllocation();
+			DestPtr += Data.GetDataAllocation();
+			SrcPtr += Data.GetDataAllocation();
 		}
 	}
 
 	bool CheckAttriubtesForRenderer();
+
+	//Debugger hooks
+#if ENABLE_NIAGARA_SIMULATION_DEBUGGING
+	FORCEINLINE void DebuggerHook_PreSpawn(FNiagaraSimulation* Sim, int32 CurrNumParticles, int32 NumToSpawn)
+	{
+		FNiagaraSimulationDebugger::Get().PreSpawn(Sim, CurrNumParticles, NumToSpawn);
+	}
+	FORCEINLINE void DebuggerHook_OnDeath(FNiagaraSimulation* Sim, int32 DyingParticle, int32 ReplacementParticle)
+	{
+		FNiagaraSimulationDebugger::Get().OnDeath(Sim, DyingParticle, ReplacementParticle);
+	}
+	FORCEINLINE void DebuggerHook_PreScriptRun(FNiagaraSimulation* Sim, UNiagaraScript* Script, int32 TotalNumParticles, int32 NumParticlesProcessed, int32 StartingParticle)
+	{
+		FNiagaraSimulationDebugger::Get().PreScriptRun(Sim, Script, TotalNumParticles, NumParticlesProcessed, StartingParticle);
+	}
+	FORCEINLINE void DebuggerHook_PostScriptRun(FNiagaraSimulation* Sim, UNiagaraScript* Script)
+	{
+		FNiagaraSimulationDebugger::Get().PostScriptRun(Sim, Script);
+	}
+#else
+	FORCEINLINE void DebuggerHook_PreSpawn(FNiagaraSimulation* Sim, int32 CurrNumParticles, int32 NumToSpawn){}
+	FORCEINLINE void DebuggerHook_OnDeath(FNiagaraSimulation* Sim, int32 DyingParticle, int32 ReplacementParticle){}
+	FORCEINLINE void DebuggerHook_PreScriptRun(FNiagaraSimulation* Sim, UNiagaraScript* Script, int32 TotalNumParticles, int32 NumParticlesProcessed, int32 StartingParticle){ }
+	FORCEINLINE void DebuggerHook_PostScriptRun(FNiagaraSimulation* Sim, UNiagaraScript* Script){}
+#endif
 };
diff --git a/Engine/Source/Runtime/Niagara/NiagaraConstantSet.h b/Engine/Source/Runtime/Niagara/NiagaraConstantSet.h
--- a/Engine/Source/Runtime/Niagara/NiagaraConstantSet.h
+++ b/Engine/Source/Runtime/Niagara/NiagaraConstantSet.h
@@ -247,7 +247,7 @@ public:
 	FORCEINLINE int32 NumScalars()const { return ScalarConstants.Num(); }
 	FORCEINLINE int32 NumVectors()const { return VectorConstants.Num(); }
 	FORCEINLINE int32 NumMatrices()const { return MatrixConstants.Num(); }
-	FORCEINLINE int32 ScalarTableSize()const { return(NumScalars() / 4) + FMath::Min(1,NumScalars() % 4); }
+	FORCEINLINE int32 ScalarTableSize()const { return(NumScalars()); }
 	FORCEINLINE int32 VectorTableSize()const { return(NumVectors()); }
 	FORCEINLINE int32 MatrixTableSize()const { return(NumMatrices() * 4); }
 
@@ -265,7 +265,7 @@ public:
 		int Idx = ScalarConstants.IndexOfByPredicate([&](const FNiagaraConstants_Float& V){ return Constant.Name == V.Name; });
 		if (Idx != INDEX_NONE)
 		{
-			return Idx / 4;
+			return Idx;
 		}
 		return INDEX_NONE;
 	}
@@ -294,7 +294,4 @@ public:
 	{ 
 		return DataObjectConstants.IndexOfByPredicate([&](const FNiagaraConstants_DataObject& V){ return Constant.Name == V.Name; });
 	}
-
-	/** Return the absolute index of the passed scalar. Can be used to get the component index also. */
-	FORCEINLINE int32 GetAbsoluteIndex_Scalar(const FNiagaraVariableInfo& Constant)const{ return ScalarConstants.IndexOfByPredicate([&](const FNiagaraConstants_Float& V){ return Constant.Name == V.Name; });; }
 };
diff --git a/Engine/Source/Runtime/Niagara/NiagaraScriptConstantData.h b/Engine/Source/Runtime/Niagara/NiagaraScriptConstantData.h
--- a/Engine/Source/Runtime/Niagara/NiagaraScriptConstantData.h
+++ b/Engine/Source/Runtime/Niagara/NiagaraScriptConstantData.h
@@ -62,15 +62,14 @@ struct FNiagaraScriptConstantData
 	Calculates the index into a a constant table created from this constant data for the given constant name.
 	Obviously, assumes the constant data is complete. Any additions etc will invalidate previously calculated indexes.
 	*/
-	void GetTableIndex(const FNiagaraVariableInfo& Constant, bool bInternal, int32& OutConstantIdx, int32& OutComponentIndex, ENiagaraDataType& OutType)const
+	void GetTableIndex(const FNiagaraVariableInfo& Constant, bool bInternal, int32& OutConstantIdx, ENiagaraDataType& OutType)const
 	{
 		int32 Base = bInternal ? ExternalConstantsNew.GetTableSize() : 0;
 		int32 ConstIdx = INDEX_NONE;
-		int32 CompIdx = INDEX_NONE;
 		ENiagaraDataType Type;
 		const FNiagaraConstants& Consts = bInternal ? InternalConstantsNew : ExternalConstantsNew;
 
-		ConstIdx = Consts.GetAbsoluteIndex_Scalar(Constant);
+		ConstIdx = Consts.GetTableIndex_Scalar(Constant);
 		Type = ENiagaraDataType::Scalar;
 		if (ConstIdx == INDEX_NONE)
 		{
@@ -89,14 +88,8 @@ struct FNiagaraScriptConstantData
 				}
 			}
 		}
-		else
-		{
-			CompIdx = ConstIdx % 4;
-			ConstIdx /= 4;
-		}
 
 		OutConstantIdx = Base + ConstIdx;
-		OutComponentIndex = CompIdx;
 		OutType = Type;
 	}
 };
diff --git a/Engine/Source/Runtime/Niagara/Private/NiagaraComponent.cpp b/Engine/Source/Runtime/Niagara/Private/NiagaraComponent.cpp
--- a/Engine/Source/Runtime/Niagara/Private/NiagaraComponent.cpp
+++ b/Engine/Source/Runtime/Niagara/Private/NiagaraComponent.cpp
@@ -20,7 +20,7 @@ DEFINE_LOG_CATEGORY(LogNiagara);
 FNiagaraSceneProxy::FNiagaraSceneProxy(const UNiagaraComponent* InComponent)
 		:	FPrimitiveSceneProxy(InComponent)
 {
-	UpdateEffectRenderers(InComponent->GetEffectInstance());
+	UpdateEffectRenderers(InComponent->GetEffectInstance().Get());
 }
 
 void FNiagaraSceneProxy::UpdateEffectRenderers(FNiagaraEffectInstance* InEffect)
@@ -149,7 +149,7 @@ void UNiagaraComponent::TickComponent(float DeltaSeconds, enum ELevelTick TickTy
 {
 //	EmitterAge += DeltaSeconds;
 
-	if (EffectInstance)
+	if (EffectInstance.Get())
 	{ 
 		static FNiagaraVariableInfo Const_Zero(TEXT("ZERO"), ENiagaraDataType::Vector);
 		static FNiagaraVariableInfo Const_DeltaTime(TEXT("Delta Time"), ENiagaraDataType::Vector);
@@ -183,32 +183,32 @@ const UObject* UNiagaraComponent::AdditionalStatObject() const
 void UNiagaraComponent::OnRegister()
 {
 	Super::OnRegister();
-	if (Asset)
+
+	if (!EffectInstance.IsValid() && Asset)
 	{
-		if (!EffectInstance)
-		{
-			EffectInstance = new FNiagaraEffectInstance(Asset, this);
-		}
-		{
-			EffectInstance->Init(this);
+		EffectInstance = MakeShareable(new FNiagaraEffectInstance(Asset, this));
+	}
 
-			// initialize all render modules
-			ENQUEUE_UNIQUE_RENDER_COMMAND_TWOPARAMETER(
-				FChangeNiagaraRenderModule,
-				FNiagaraEffectInstance*, InEffect, this->EffectInstance,
-				UNiagaraComponent*, InComponent, this,
-				{
-					for (TSharedPtr<FNiagaraSimulation> Emitter : InEffect->GetEmitters())
-					{
-						Emitter->SetRenderModuleType(Emitter->GetProperties()->RenderModuleType, InComponent->GetWorld()->FeatureLevel);
-					}
-					InEffect->RenderModuleupdate();
-				}
-			);
-			FlushRenderingCommands();
+	if (EffectInstance.IsValid())
+	{
+		EffectInstance->Init(this);
+
+		// initialize all render modules
+		ENQUEUE_UNIQUE_RENDER_COMMAND_TWOPARAMETER(
+			FChangeNiagaraRenderModule,
+			FNiagaraEffectInstance*, InEffect, this->EffectInstance.Get(),
+			UNiagaraComponent*, InComponent, this,
+			{
+			for (TSharedPtr<FNiagaraSimulation> Emitter : InEffect->GetEmitters())
+			{
+				Emitter->SetRenderModuleType(Emitter->GetProperties()->RenderModuleType, InComponent->GetWorld()->FeatureLevel);
+			}
+			InEffect->RenderModuleupdate();
 		}
-		VectorVM::Init();
+		);
+		FlushRenderingCommands();
 	}
+	VectorVM::Init();
 }
 
 
@@ -219,7 +219,7 @@ void UNiagaraComponent::OnUnregister()
 
 void UNiagaraComponent::SendRenderDynamicData_Concurrent()
 {
-	if (EffectInstance && SceneProxy)
+	if (EffectInstance.IsValid() && SceneProxy)
 	{
 		FNiagaraSceneProxy* NiagaraProxy = static_cast<FNiagaraSceneProxy*>(SceneProxy);
 		for (int32 i = 0; i < EffectInstance->GetEmitters().Num(); i++)
@@ -264,7 +264,7 @@ int32 UNiagaraComponent::GetNumMaterials() const
 FBoxSphereBounds UNiagaraComponent::CalcBounds(const FTransform& LocalToWorld) const
 {
 	FBox SimBounds(ForceInit);
-	if (EffectInstance)
+	if (EffectInstance.IsValid())
 	{
 		FBoxSphereBounds BSBounds(EffectInstance->GetEffectBounds());
 		//BSBounds.Origin = LocalToWorld.GetLocation();
@@ -276,7 +276,7 @@ FBoxSphereBounds UNiagaraComponent::CalcBounds(const FTransform& LocalToWorld) c
 FPrimitiveSceneProxy* UNiagaraComponent::CreateSceneProxy()
 {
 	FNiagaraSceneProxy* Proxy = new FNiagaraSceneProxy(this);
-	Proxy->UpdateEffectRenderers(EffectInstance);
+	Proxy->UpdateEffectRenderers(EffectInstance.Get());
 	return Proxy;
 }
 
@@ -295,7 +295,7 @@ void UNiagaraComponent::SetAsset(UNiagaraEffect* InAsset)
 {
 	Asset = InAsset;
 
-	EffectInstance = new FNiagaraEffectInstance(Asset, this);
+	EffectInstance = MakeShareable(new FNiagaraEffectInstance(Asset, this));
 }
 
 const TArray<FNiagaraVariableInfo>& UNiagaraComponent::GetSystemConstants()
@@ -321,4 +321,7 @@ const TArray<FNiagaraVariableInfo>& UNiagaraComponent::GetSystemConstants()
 		SystemConstants.Add(Const_EmitterTransform);
 	}
 	return SystemConstants;
-}
\ No newline at end of file
+}
+FNiagaraDataSetID FNiagaraDataSetID::DeathEvent(TEXT("Death"),ENiagaraDataSetType::Event);
+FNiagaraDataSetID FNiagaraDataSetID::SpawnEvent(TEXT("Spawn"), ENiagaraDataSetType::Event);
+FNiagaraDataSetID FNiagaraDataSetID::CollisionEvent(TEXT("Collision"), ENiagaraDataSetType::Event);
\ No newline at end of file
diff --git a/Engine/Source/Runtime/Niagara/Private/NiagaraConstantSet.cpp b/Engine/Source/Runtime/Niagara/Private/NiagaraConstantSet.cpp
--- a/Engine/Source/Runtime/Niagara/Private/NiagaraConstantSet.cpp
+++ b/Engine/Source/Runtime/Niagara/Private/NiagaraConstantSet.cpp
@@ -258,7 +258,8 @@ void FNiagaraConstants::Init(UNiagaraEmitterProperties* EmitterProps, FNiagaraEm
 		else
 		{
 			//Otherwise, duplicate the data from the script.
-			DataObjectConstants[AddIndex].Value = CastChecked<UNiagaraDataObject>(StaticDuplicateObject(ScriptConst.Value, EmitterProps, NULL));
+			if (ScriptConst.Value)
+				DataObjectConstants[AddIndex].Value = CastChecked<UNiagaraDataObject>(StaticDuplicateObject(ScriptConst.Value, EmitterProps, NULL));
 		}
 	}
 }
@@ -299,12 +300,9 @@ void FNiagaraConstants::AppendToConstantsTable(TArray<FVector4>& ConstantsTable)
 	ConstantsTable.AddUninitialized(ScalarTableSize() + VectorTableSize() + MatrixTableSize());
 	for (const FNiagaraConstants_Float& Sc : ScalarConstants)
 	{
-		ConstantsTable[Idx + (NewIdx / 4)][NewIdx % 4] = Sc.Value;
-		++NewIdx;
-		Idx = NewIdx % 4 == 0 ? Idx + 1 : Idx;
+		ConstantsTable[Idx++] = FVector4(Sc.Value, Sc.Value, Sc.Value, Sc.Value);
 	}
 
-	Idx += FMath::Min(1, NewIdx % 4);//Move to the next table entry if needed.
 	for (const FNiagaraConstants_Vector& Vc : VectorConstants)
 	{
 		ConstantsTable[Idx++] = Vc.Value;
@@ -331,18 +329,14 @@ void FNiagaraConstants::AppendBufferConstants(TArray<class UNiagaraDataObject*>
 void FNiagaraConstants::AppendToConstantsTable(TArray<FVector4>& ConstantsTable, const FNiagaraConstantMap& Externals)const
 {
 	int32 Idx = ConstantsTable.Num();
-	int32 NewIdx = 0;
 	ConstantsTable.AddUninitialized(ScalarTableSize() + VectorTableSize() + MatrixTableSize());
 	for (int32 i = 0; i < ScalarConstants.Num(); ++i)
 	{
 		const float* Scalar = Externals.FindScalar(ScalarConstants[i].Name);
-		ConstantsTable[Idx + (NewIdx / 4)][NewIdx % 4] = Scalar ? *Scalar : ScalarConstants[i].Value;
-		++NewIdx;
-		Idx = NewIdx % 4 == 0 ? Idx + 1 : Idx;
+		float Value = Scalar ? *Scalar : ScalarConstants[i].Value;
+		ConstantsTable[Idx++] = FVector4(Value, Value, Value, Value);
 	}
-
-	Idx += FMath::Min(1, NewIdx % 4);//Move to the next table entry if needed.
-
+	
 	for (int32 i = 0; i < VectorConstants.Num(); ++i)
 	{
 		const FVector4* Vector = Externals.FindVector(VectorConstants[i].Name);
diff --git a/Engine/Source/Runtime/Niagara/Private/NiagaraEffect.cpp b/Engine/Source/Runtime/Niagara/Private/NiagaraEffect.cpp
--- a/Engine/Source/Runtime/Niagara/Private/NiagaraEffect.cpp
+++ b/Engine/Source/Runtime/Niagara/Private/NiagaraEffect.cpp
@@ -43,11 +43,11 @@ void UNiagaraEffect::CreateEffectRendererProps(TSharedPtr<FNiagaraSimulation> Si
 	Sim->GetEffectRenderer()->SetRendererProperties(Sim->GetProperties()->RendererProperties);
 }
 
-void UNiagaraEffect::Serialize(FArchive& Ar)
+void UNiagaraEffect::PostLoad()
 {
-	Super::Serialize(Ar);
+	Super::PostLoad();
 
-	if (Ar.IsLoading() && Ar.UE4Ver() < VER_UE4_NIAGARA_DATA_OBJECT_DEV_UI_FIX)
+	if (GetLinkerUE4Version() < VER_UE4_NIAGARA_DATA_OBJECT_DEV_UI_FIX)
 	{
 		int32 NumEmitters = EmitterPropsSerialized_DEPRECATED.Num();
 		for (int32 i = 0; i < NumEmitters; ++i)
@@ -63,7 +63,7 @@ void UNiagaraEffect::Serialize(FArchive& Ar)
 
 TSharedPtr<FNiagaraSimulation> FNiagaraEffectInstance::AddEmitter(UNiagaraEmitterProperties *Properties)
 {
-	FNiagaraSimulation *SimPtr = new FNiagaraSimulation(Properties, Effect);
+	FNiagaraSimulation *SimPtr = new FNiagaraSimulation(Properties, this);
 	TSharedPtr<FNiagaraSimulation> Sim = MakeShareable(SimPtr);
 	Sim->SetRenderModuleType(Properties->RenderModuleType, Component->GetWorld()->FeatureLevel);
 	Emitters.Add(Sim);
@@ -88,6 +88,11 @@ void FNiagaraEffectInstance::Tick(float DeltaSeconds)
 	// TODO: should probably just pass a pointer to the table
 	EffectBounds.Init();
 
+	for (TPair<FNiagaraDataSetID, FNiagaraDataSet>& EventSetPair : ExternalEvents)
+	{
+		EventSetPair.Value.Tick();
+	}
+
 	for (TSharedPtr<FNiagaraSimulation>&it : Emitters)
 	{
 		UNiagaraEmitterProperties *Props = it->GetProperties().Get();
@@ -119,12 +124,18 @@ void FNiagaraEffectInstance::Tick(float DeltaSeconds)
 			}
 		}
 
+		//TODO - Handle constants better. Like waaaay better.
+		it->GetConstants().Merge(it->GetProperties()->SpawnScriptProps.ExternalConstants);
+		it->GetConstants().Merge(it->GetProperties()->UpdateScriptProps.ExternalConstants);
+		it->GetConstants().Merge(Constants);
+
+		it->PreTick();
+	}
+
+	for (TSharedPtr<FNiagaraSimulation>&it : Emitters)
+	{
 		if (it->GetTickState() != NTS_Dead && it->GetTickState() != NTS_Suspended)
 		{
-			//TODO - Handle constants better. Like waaaay better.
-			it->SetConstants(Constants);
-			it->GetConstants().Merge(it->GetProperties()->SpawnScriptProps.ExternalConstants);
-			it->GetConstants().Merge(it->GetProperties()->UpdateScriptProps.ExternalConstants);
 			it->Tick(DeltaSeconds);
 		}
 
@@ -134,6 +145,19 @@ void FNiagaraEffectInstance::Tick(float DeltaSeconds)
 	Age += DeltaSeconds;
 }
 
+FNiagaraSimulation* FNiagaraEffectInstance::GetEmitter(FString Name)
+{
+	for (TSharedPtr<FNiagaraSimulation> Sim : Emitters)
+	{
+		UNiagaraEmitterProperties* Props = Sim->GetProperties().Get();
+		if (Props && Props->EmitterName == Name)
+		{
+			return Sim.Get();//We really need to sort out the ownership of objects in naigara. Am I free to pass a raw ptr here?
+		}
+	}
+	return NULL;
+}
+
 void FNiagaraEffectInstance::RenderModuleupdate()
 {
 	FNiagaraSceneProxy *NiagaraProxy = static_cast<FNiagaraSceneProxy*>(Component->SceneProxy);
@@ -143,14 +167,62 @@ void FNiagaraEffectInstance::RenderModuleupdate()
 	}
 }
 
-
 void FNiagaraEffectInstance::InitEmitters(UNiagaraEffect *InAsset)
 {
 	check(InAsset);
 	for (int i = 0; i < InAsset->GetNumEmitters(); i++)
 	{
 		UNiagaraEmitterProperties *Props = InAsset->GetEmitterProperties(i);
-		FNiagaraSimulation *Sim = new FNiagaraSimulation(Props, InAsset);
+		Props->Init();//Init these to be sure any cached data from the scripts is up to date.
+
+		FNiagaraSimulation *Sim = new FNiagaraSimulation(Props, this);
 		Emitters.Add(MakeShareable(Sim));
 	}
+
+	for (TSharedPtr<FNiagaraSimulation> Emitter : Emitters)
+	{
+		Emitter->PostInit();
+	}
+}
+
+void FNiagaraEffectInstance::ReInitEmitters()
+{
+	for (TSharedPtr<FNiagaraSimulation> Emitter : Emitters)
+	{
+		UNiagaraEmitterProperties* Props = Emitter->GetProperties().Get();
+		check(Props);
+		Props->Init();//Init these to be sure any cached data from the scripts is up to date.
+		Emitter->Init();
+	}
+
+	for (TSharedPtr<FNiagaraSimulation> Emitter : Emitters)
+	{
+		Emitter->PostInit();
+	}
 }
+
+FNiagaraDataSet* FNiagaraEffectInstance::GetDataSet(FNiagaraDataSetID SetID, FName EmitterName)
+{
+	if (EmitterName == NAME_None)
+	{
+		if (FNiagaraDataSet* ExternalSet = ExternalEvents.Find(SetID))
+		{
+			return ExternalSet;
+		}
+	}
+	for (TSharedPtr<FNiagaraSimulation> Emitter : Emitters)
+	{
+		check(Emitter.IsValid());
+		if (Emitter->IsEnabled())
+		{
+			UNiagaraEmitterProperties* PinnedProps = Emitter->GetProperties().Get();
+			check(PinnedProps);
+			if (*PinnedProps->EmitterName == EmitterName)
+			{
+				return Emitter->GetDataSet(SetID);
+			}
+		}
+	}
+
+	return NULL;
+}
\ No newline at end of file
diff --git a/Engine/Source/Runtime/Niagara/Private/NiagaraEmitterProperties.cpp b/Engine/Source/Runtime/Niagara/Private/NiagaraEmitterProperties.cpp
new file mode 100644
--- /dev/null
+++ b/Engine/Source/Runtime/Niagara/Private/NiagaraEmitterProperties.cpp
@@ -0,0 +1,148 @@
+// Copyright 1998-2015 Epic Games, Inc. All Rights Reserved.
+
+
+#include "NiagaraPrivate.h"
+#include "NiagaraEmitterProperties.h"
+#include "NiagaraScriptSourceBase.h"
+
+void FNiagaraEmitterScriptProperties::Init(UNiagaraEmitterProperties* EmitterProps)
+{
+	if (Script)
+	{
+		//TODO - Don't need to do this for every init. Only when the script has changed.
+		//Store a guid or something
+
+		ExternalConstants.Init(EmitterProps, this);
+
+		TArray<FNiagaraEventGeneratorProperties> NewEventGenerators;
+		TArray<FNiagaraDataSetProperties>& EventGeneratorDataSets = Script->EventGenerators;
+		for (int32 i = 0; i < EventGeneratorDataSets.Num(); ++i)
+		{
+			FNiagaraDataSetProperties& CurrGeneratorSet = EventGeneratorDataSets[i];
+
+			//If we already had customized data for this generator then use that, otherwise add a new defaulted one.
+			FNiagaraEventGeneratorProperties* ExistingEventGenerator = EventGenerators.FindByPredicate(
+				[&](FNiagaraEventGeneratorProperties& Generator)
+			{
+				return Generator.ID == CurrGeneratorSet.ID;
+			}
+			);
+
+			if (ExistingEventGenerator)
+			{
+				NewEventGenerators.Add(*ExistingEventGenerator);
+			}
+			else
+			{
+				FNiagaraEventGeneratorProperties NewGenerator;
+				NewGenerator.ID = CurrGeneratorSet.ID;
+				NewGenerator.Variables = CurrGeneratorSet.Variables;
+				NewEventGenerators.Add(NewGenerator);
+			}
+		}
+		EventGenerators = NewEventGenerators;
+
+		TArray<FNiagaraEventReceiverProperties> NewEventReceivers;
+		TArray<FNiagaraDataSetProperties>& EventReceiverDataSets = Script->EventReceivers;
+		for (int32 i = 0; i < EventReceiverDataSets.Num(); ++i)
+		{
+			FNiagaraDataSetProperties& CurrReceiverSet = EventReceiverDataSets[i];
+
+			//If we already had customized data for this receiver then use that, otherwise add a new defaulted one.
+			FNiagaraEventReceiverProperties* ExistingEventReceiver = EventReceivers.FindByPredicate(
+				[&](FNiagaraEventReceiverProperties& Receiver)
+			{
+				return Receiver.Name == CurrReceiverSet.ID.Name;
+			}
+			);
+
+			if (ExistingEventReceiver)
+			{
+				NewEventReceivers.Add(*ExistingEventReceiver);
+			}
+			else
+			{
+				FNiagaraEventReceiverProperties NewEventReceiver;
+				NewEventReceiver.Name = CurrReceiverSet.ID.Name;
+				NewEventReceivers.Add(NewEventReceiver);
+			}
+		}
+
+		EventReceivers = NewEventReceivers;
+	}
+	else
+	{
+		ExternalConstants.Empty();
+		EventReceivers.Empty();
+		EventGenerators.Empty();
+	}
+}
+
+//////////////////////////////////////////////////////////////////////////
+
+UNiagaraEmitterProperties::UNiagaraEmitterProperties(const FObjectInitializer& Initiilaizer)
+: Super(Initiilaizer)
+, EmitterName(TEXT("New Emitter"))
+, bIsEnabled(true)
+, SpawnRate(50)
+, Material(nullptr)
+, RenderModuleType(RMT_Sprites)
+, StartTime(0.0f)
+, EndTime(0.0f)
+, RendererProperties(nullptr)
+, NumLoops(0)
+{
+}
+
+void UNiagaraEmitterProperties::Init()
+{
+	SpawnScriptProps.Init(this);
+	UpdateScriptProps.Init(this);
+}
+
+void UNiagaraEmitterProperties::InitFromOldStruct(FDeprecatedNiagaraEmitterProperties& OldStruct)
+{
+	EmitterName = OldStruct.Name;
+	bIsEnabled = OldStruct.bIsEnabled;
+	SpawnRate = OldStruct.SpawnRate;
+	UpdateScriptProps.Script = OldStruct.UpdateScript;
+	SpawnScriptProps.Script = OldStruct.SpawnScript;
+	Material = OldStruct.Material;
+	RenderModuleType = OldStruct.RenderModuleType;
+	StartTime = OldStruct.StartTime;
+	EndTime = OldStruct.EndTime;
+	RendererProperties = OldStruct.RendererProperties;
+	NumLoops = OldStruct.NumLoops;
+
+#if WITH_EDITORONLY_DATA
+	//Ensure both scripts are post loaded and compile them so that their constant arrays are correct.
+	if (SpawnScriptProps.Script)
+	{
+		SpawnScriptProps.Script->ConditionalPostLoad();
+		if (SpawnScriptProps.Script->Source)
+			SpawnScriptProps.Script->Source->Compile();
+	}
+	if (UpdateScriptProps.Script)
+	{
+		UpdateScriptProps.Script->ConditionalPostLoad();
+		if (UpdateScriptProps.Script->Source)
+			UpdateScriptProps.Script->Source->Compile();
+	}
+#endif
+
+	UpdateScriptProps.ExternalConstants.InitFromOldMap(OldStruct.ExternalConstants);
+	SpawnScriptProps.ExternalConstants.InitFromOldMap(OldStruct.ExternalSpawnConstants);
+	//Also init the spawn constants with the update constants as they used to combined
+	SpawnScriptProps.ExternalConstants.InitFromOldMap(OldStruct.ExternalConstants);
+	UpdateScriptProps.ExternalConstants.InitFromOldMap(OldStruct.ExternalSpawnConstants);
+
+	Init();
+}
+
+#if WITH_EDITOR
+void UNiagaraEmitterProperties::PostEditChangeProperty(struct FPropertyChangedEvent& PropertyChangedEvent)
+{
+	Super::PostEditChangeProperty(PropertyChangedEvent);
+	Init();
+}
+#endif
\ No newline at end of file
diff --git a/Engine/Source/Runtime/Niagara/Private/NiagaraEvents.cpp b/Engine/Source/Runtime/Niagara/Private/NiagaraEvents.cpp
new file mode 100644
--- /dev/null
+++ b/Engine/Source/Runtime/Niagara/Private/NiagaraEvents.cpp
@@ -0,0 +1,77 @@
+
+// Copyright 1998-2015 Epic Games, Inc. All Rights Reserved.
+
+
+#include "NiagaraPrivate.h"
+#include "NiagaraEvents.h"
+#include "NiagaraSimulation.h"
+#include "NiagaraEffect.h"
+
+DECLARE_DWORD_COUNTER_STAT(TEXT("Num Death Events"), STAT_NiagaraNumDeathEvents, STATGROUP_Niagara);
+DECLARE_DWORD_COUNTER_STAT(TEXT("Num Spawn Events"), STAT_NiagaraNumSpawnEvents, STATGROUP_Niagara);
+
+void UNiagaraEventReceiverEmitterAction_SpawnParticles::PerformAction(FNiagaraSimulation& OwningSim, FNiagaraEventReceiverProperties& OwningEventReceiver)
+{
+	//Find the generator set we're bound to and spawn NumParticles particles for every event it generated last frame.
+	FNiagaraDataSet* GeneratorSet = OwningSim.GetParentEffectInstance()->GetDataSet(FNiagaraDataSetID(OwningEventReceiver.SourceEventGenerator, ENiagaraDataSetType::Event), OwningEventReceiver.SourceEmitter);
+	if (GeneratorSet)
+	{
+		OwningSim.SpawnBurst(GeneratorSet->GetPrevNumInstances()*NumParticles);
+	}
+}
+//////////////////////////////////////////////////////////////////////////
+
+void FNiagaraDeathEventGenerator::BeginTrackingDeaths()
+{
+	FNiagaraDataSet& Data = OwnerSim->GetData();
+	TempStorage.Empty(Data.GetNumInstances());
+	TempStorage.SetNum(Data.GetNumVariables());
+	NumDead = 0;
+}
+
+void FNiagaraDeathEventGenerator::OnDeath(int32 DeadIndex)
+{
+	FNiagaraDataSet& Data = OwnerSim->GetData();
+	int32 AttrIdx = 0;
+	for (const TPair<FNiagaraVariableInfo, uint32>& Pair : Data.GetVariables())
+	{
+		FVector4* AttrPtr = Data.GetVariableData(Pair.Key);
+		TempStorage[AttrIdx].Add(AttrPtr[DeadIndex]);
+		++AttrIdx;
+	}
+	++NumDead;
+}
+
+void FNiagaraDeathEventGenerator::EndTrackingDeaths()
+{
+	FNiagaraDataSet& Data = OwnerSim->GetData();
+	int32 AttrIdx = 0;
+	DataSet.SetNumInstances(NumDead);
+	DataSet.Allocate(NumDead);
+	if (NumDead > 0)
+	{
+		INC_DWORD_STAT_BY(STAT_NiagaraNumDeathEvents, NumDead);
+		for (const TPair<FNiagaraVariableInfo, uint32>& Pair : Data.GetVariables())
+		{
+			FVector4* AttrPtr = DataSet.GetVariableData(Pair.Key);
+			FMemory::Memcpy(AttrPtr, TempStorage[AttrIdx].GetData(), sizeof(FVector4)* NumDead);
+			++AttrIdx;
+		}
+	}
+}
+
+void FNiagaraSpawnEventGenerator::OnSpawned(int32 SpawnIndex, int32 NumSpawned)
+{
+	INC_DWORD_STAT_BY(STAT_NiagaraNumSpawnEvents, NumSpawned);
+
+	FNiagaraDataSet& Data = OwnerSim->GetData();
+	DataSet.Allocate(NumSpawned);
+	DataSet.SetNumInstances(NumSpawned);
+
+	for (const TPair<FNiagaraVariableInfo, uint32>& Pair : Data.GetVariables())
+	{
+		FVector4* DestPtr = DataSet.GetVariableData(Pair.Key, 0);
+		FVector4* AttrPtr = Data.GetVariableData(Pair.Key, SpawnIndex);
+		FMemory::Memcpy(DestPtr, AttrPtr, sizeof(FVector4)* NumSpawned);
+	}
+}
diff --git a/Engine/Source/Runtime/Niagara/Private/NiagaraSimulation.cpp b/Engine/Source/Runtime/Niagara/Private/NiagaraSimulation.cpp
--- a/Engine/Source/Runtime/Niagara/Private/NiagaraSimulation.cpp
+++ b/Engine/Source/Runtime/Niagara/Private/NiagaraSimulation.cpp
@@ -4,29 +4,44 @@
 #include "NiagaraSimulation.h"
 #include "NiagaraEffectRenderer.h"
 #include "VectorVM.h"
+#include "NiagaraEffect.h"
 
-FNiagaraSimulation::FNiagaraSimulation(TWeakObjectPtr<UNiagaraEmitterProperties> InProps, UNiagaraEffect* InEffect)
+DECLARE_DWORD_COUNTER_STAT(TEXT("Num Custom Events"), STAT_NiagaraNumCustomEvents, STATGROUP_Niagara);
+
+//////////////////////////////////////////////////////////////////////////
+
+FNiagaraSimulation::FNiagaraSimulation(TWeakObjectPtr<UNiagaraEmitterProperties> InProps, FNiagaraEffectInstance* InParentEffectInstance)
 : Age(0.0f)
 , Loops(0)
 , bIsEnabled(true)
+, Data(FNiagaraDataSetID(*InProps->EmitterName, ENiagaraDataSetType::ParticleData))
 , SpawnRemainder(0.0f)
 , CachedBounds(ForceInit)
 , EffectRenderer(nullptr)
-, ParentEffect(InEffect)
+, ParentEffectInstance(InParentEffectInstance)
+, SpawnEventGenerator(this)
+, bGenerateSpawnEvents(false)
+, DeathEventGenerator(this)
+, bGenerateDeathEvents(false)
 {
 	Props = InProps;
 
 	Init();
 }
 
-FNiagaraSimulation::FNiagaraSimulation(TWeakObjectPtr<UNiagaraEmitterProperties> InProps, UNiagaraEffect* InEffect, ERHIFeatureLevel::Type InFeatureLevel)
-	: Age(0.0f)
-	, Loops(0)
-	, bIsEnabled(true)
-	, SpawnRemainder(0.0f)
-	, CachedBounds(ForceInit)
-	, EffectRenderer(nullptr)
-	, ParentEffect(InEffect)
+FNiagaraSimulation::FNiagaraSimulation(TWeakObjectPtr<UNiagaraEmitterProperties> InProps, FNiagaraEffectInstance* InParentEffectInstance, ERHIFeatureLevel::Type InFeatureLevel)
+: Age(0.0f)
+, Loops(0)
+, bIsEnabled(true)
+, Data(FNiagaraDataSetID(*InProps->EmitterName, ENiagaraDataSetType::ParticleData))
+, SpawnRemainder(0.0f)
+, CachedBounds(ForceInit)
+, EffectRenderer(nullptr)
+, ParentEffectInstance(InParentEffectInstance)
+, SpawnEventGenerator(this)
+, bGenerateSpawnEvents(false)
+, DeathEventGenerator(this)
+, bGenerateDeathEvents(false)
 {
 	Props = InProps;
 
@@ -37,46 +52,114 @@ FNiagaraSimulation::FNiagaraSimulation(TWeakObjectPtr<UNiagaraEmitterProperties>
 void FNiagaraSimulation::Init()
 {
 	Data.Reset();
-	SpawnRemainder = 0.0f;
-	Constants.Empty();
-	Age = 0.0f;
-	Loops = 0.0f;
+	DataSets.Empty();
+	DataSetMap.Empty();
+	SpawnEventGenerator.Reset();
+	DeathEventGenerator.Reset();
+	bGenerateSpawnEvents = false;
+	bGenerateDeathEvents = false;
+
+	if (!bIsEnabled)
+	{
+		return;
+	}
 
 	UNiagaraEmitterProperties* PinnedProps = Props.Get();
-	if (PinnedProps && PinnedProps->UpdateScriptProps.Script && PinnedProps->SpawnScriptProps.Script)
+	if (!PinnedProps)
+	{
+		UE_LOG(LogNiagara, Error, TEXT("Unknown Error creating Niagara Simulation. Properties were null."));
+		bIsEnabled = false;
+		return;
+	}
+
+	//Check for various failure conditions and bail.
+	if (!PinnedProps || !PinnedProps->UpdateScriptProps.Script || !PinnedProps->SpawnScriptProps.Script)
+	{
+		//TODO - Arbitrary named scripts. Would need some base functionality for Spawn/Udpate to be called that can be overriden in BPs for emitters with custom scripts.
+		UE_LOG(LogNiagara, Error, TEXT("This emitter cannot be enabled because it's doesn't have both an update and spawn script."));
+		bIsEnabled = false;
+		return;
+	}
+	if (PinnedProps->SpawnScriptProps.Script->Usage.bReadsAttriubteData)
 	{
-		if (PinnedProps->UpdateScriptProps.Script->Attributes.Num() == 0 || PinnedProps->SpawnScriptProps.Script->Attributes.Num() == 0)
+		UE_LOG(LogNiagara, Error, TEXT("%s reads attribute data and so cannot be used as a spawn script. The data being read would be invalid."), *Props->SpawnScriptProps.Script->GetName());
+		bIsEnabled = false;
+		return;
+	}
+	if (PinnedProps->UpdateScriptProps.Script->Attributes.Num() == 0 || Props->SpawnScriptProps.Script->Attributes.Num() == 0)
+	{
+		UE_LOG(LogNiagara, Error, TEXT("This emitter cannot be enabled because it's spawn or update script doesn't have any attriubtes.."));
+		bIsEnabled = false;
+		return;
+	}
+	
+	//Add the particle data to the data set map.
+	//Currently just used for the tick loop but will also allow access directly to the particle data from other emitters.
+	DataSetMap.Add(Data.GetID(), &Data);
+	//Warn the user if there are any attributes used in the update script that are not initialized in the spawn script.
+	//TODO: We need some window in the effect editor and possibly the graph editor for warnings and errors.
+	for (FNiagaraVariableInfo& Attr : PinnedProps->UpdateScriptProps.Script->Attributes)
+	{
+		int32 FoundIdx;
+		if (!PinnedProps->SpawnScriptProps.Script->Attributes.Find(Attr, FoundIdx))
 		{
-			Data.Reset();
-			bIsEnabled = false;
-			UE_LOG(LogNiagara, Error, TEXT("This emitter cannot be enabled because it's spawn or update script doesn't have any attriubtes.."));
+			UE_LOG(LogNiagara, Warning, TEXT("Attribute %s is used in the Update script for %s but it is not initialised in the Spawn script!"), *Attr.Name.ToString(), *Props->EmitterName);
 		}
-		else
+	}
+	Data.AddVariables(PinnedProps->UpdateScriptProps.Script->Attributes);
+	Data.AddVariables(PinnedProps->SpawnScriptProps.Script->Attributes);
+
+	Constants.Empty();
+	Constants.Merge(PinnedProps->UpdateScriptProps.ExternalConstants);
+	Constants.Merge(PinnedProps->SpawnScriptProps.ExternalConstants);
+
+	//Create data storage for event generators.
+	for (FNiagaraEventGeneratorProperties& Generator : Props->SpawnScriptProps.EventGenerators)
+	{
+		FNiagaraDataSet* DataSet = DataSetMap.FindOrAdd(Generator.ID);
+		if (!DataSet)
 		{
-			//Warn the user if there are any attributes used in the update script that are not initialized in the spawn script.
-			//TODO: We need some window in the effect editor and possibly the graph editor for warnings and errors.
-			for (FNiagaraVariableInfo& Attr : PinnedProps->UpdateScriptProps.Script->Attributes)
-			{
-				int32 FoundIdx;
-				if (!PinnedProps->SpawnScriptProps.Script->Attributes.Find(Attr, FoundIdx))
-				{
-					UE_LOG(LogNiagara, Warning, TEXT("Attribute %s is used in the Update script for %s but it is not initialised in the Spawn script!"), *Attr.Name.ToString(), *Props->EmitterName);
-				}
-			}
-			Data.AddAttributes(PinnedProps->UpdateScriptProps.Script->Attributes);
-			Data.AddAttributes(PinnedProps->SpawnScriptProps.Script->Attributes);
+			int32 Index = DataSets.Add(FNiagaraDataSet(Generator.ID));
+			DataSet = &DataSets[Index];
+			DataSetMap.Add(Generator.ID, DataSet);
+		}
+
+		check(DataSet);
+		DataSet->AddVariables(Generator.Variables);
+		DataSet->Allocate(FMath::Max(DataSet->GetDataAllocation(), (uint32)Generator.MaxEventsPerFrame));
+	}
 
-			Constants.Empty();
-			Constants.Merge(PinnedProps->UpdateScriptProps.ExternalConstants);
-			Constants.Merge(PinnedProps->SpawnScriptProps.ExternalConstants);
+	for (FNiagaraEventGeneratorProperties& Generator : Props->UpdateScriptProps.EventGenerators)
+	{
+		FNiagaraDataSet* DataSet = DataSetMap.FindOrAdd(Generator.ID);
+		if (!DataSet)
+		{
+			int32 Index = DataSets.Add(FNiagaraDataSet(Generator.ID));
+			DataSet = &DataSets[Index];
+			DataSetMap.Add(Generator.ID, DataSet);
 		}
+
+		check(DataSet);
+		DataSet->AddVariables(Generator.Variables);
+		DataSet->Allocate(FMath::Max(DataSet->GetDataAllocation(), (uint32)Generator.MaxEventsPerFrame));
 	}
-	else
+}
+
+void FNiagaraSimulation::PostInit()
+{
+	if (bIsEnabled)
 	{
-		Data.Reset();//Clear the attributes to mark the sim as disabled independatly of the user set Enabled flag.
-		bIsEnabled = false;//Also force the user flag to give an indication to the user.
-		//TODO - Arbitrary named scripts. Would need some base functionality for Spawn/Udpate to be called that can be overriden in BPs for emitters with custom scripts.
-		UE_LOG(LogNiagara, Error, TEXT("This emitter cannot be enabled because it's doesn't have both an update and spawn script."));
+		check(ParentEffectInstance);
+		//Go through all our receivers and grab their generator sets so that the source emitters can do any init work they need to do.
+		for (FNiagaraEventReceiverProperties& Receiver : Props->SpawnScriptProps.EventReceivers)
+		{
+			FNiagaraDataSet* ReceiverSet = ParentEffectInstance->GetDataSet(FNiagaraDataSetID(Receiver.SourceEventGenerator, ENiagaraDataSetType::Event), Receiver.SourceEmitter);
+		}
+
+		for (FNiagaraEventReceiverProperties& Receiver : Props->UpdateScriptProps.EventReceivers)
+		{
+			FNiagaraDataSet* ReceiverSet = ParentEffectInstance->GetDataSet(FNiagaraDataSetID(Receiver.SourceEventGenerator, ENiagaraDataSetType::Event), Receiver.SourceEmitter);
+		}
 	}
 }
 
@@ -92,6 +175,38 @@ void FNiagaraSimulation::SetEnabled(bool bEnabled)
 	Init();
 }
 
+FNiagaraDataSet* FNiagaraSimulation::GetDataSet(FNiagaraDataSetID SetID)
+{
+	FNiagaraDataSet** SetPtr = DataSetMap.Find(SetID);
+	FNiagaraDataSet* Ret = NULL;
+	if (SetPtr)
+	{
+		Ret = *SetPtr;
+	}
+	else
+	{
+		//If the set was not found and we're looking for a system event generator then init them here.
+		if (SetID == FNiagaraDataSetID::SpawnEvent)
+		{
+			bGenerateSpawnEvents = true;
+			SpawnEventGenerator.GetDataSet().AddVariables(Props->SpawnScriptProps.Script->Attributes);
+			SpawnEventGenerator.GetDataSet().AddVariables(Props->UpdateScriptProps.Script->Attributes);
+			DataSetMap.Add(SpawnEventGenerator.GetDataSet().GetID(), &SpawnEventGenerator.GetDataSet());
+			Ret = &SpawnEventGenerator.GetDataSet();
+		}
+		else if (SetID == FNiagaraDataSetID::DeathEvent)
+		{
+			bGenerateDeathEvents = true;
+			DeathEventGenerator.GetDataSet().AddVariables(Props->SpawnScriptProps.Script->Attributes);
+			DeathEventGenerator.GetDataSet().AddVariables(Props->UpdateScriptProps.Script->Attributes);
+			DataSetMap.Add(DeathEventGenerator.GetDataSet().GetID(), &DeathEventGenerator.GetDataSet());
+			Ret = &DeathEventGenerator.GetDataSet();
+		}
+		//TODO: Scene Collision Event Generator.
+	}
+
+	return Ret;
+}
 
 float FNiagaraSimulation::GetTotalCPUTime()
 {
@@ -106,9 +221,78 @@ float FNiagaraSimulation::GetTotalCPUTime()
 
 int FNiagaraSimulation::GetTotalBytesUsed()
 {
-	return Data.GetBytesUsed();
+	int32 BytesUsed = Data.GetBytesUsed();
+	for (FNiagaraDataSet& Set : DataSets)
+	{
+		BytesUsed += Set.GetBytesUsed();
+	}
+	BytesUsed += DeathEventGenerator.GetDataSet().GetBytesUsed();
+	BytesUsed += SpawnEventGenerator.GetDataSet().GetBytesUsed();
+	return BytesUsed;
 }
 
+void FNiagaraSimulation::PreTick()
+{
+	UNiagaraEmitterProperties* PinnedProps = Props.Get();
+	if (!PinnedProps || !bIsEnabled || TickState == NTS_Suspended || TickState == NTS_Dead)
+		return;
+
+	check(Data.GetNumVariables() > 0);
+	check(PinnedProps->SpawnScriptProps.Script);
+	check(PinnedProps->UpdateScriptProps.Script);
+
+	// Iterate over looking for dead particles and move from the end of the list to the dead location, compacting in the process
+	{
+		SCOPE_CYCLE_COUNTER(STAT_NiagaraKill);
+		int32 OrigNumParticles = Data.GetNumInstances();
+		int32 CurNumParticles = OrigNumParticles;
+		int32 ParticleIndex = OrigNumParticles-1;
+
+		if (bGenerateDeathEvents)
+		{
+			DeathEventGenerator.BeginTrackingDeaths();
+		}
+		const FVector4* ParticleRelativeTimes = Data.GetVariableData(FNiagaraVariableInfo(FName(TEXT("Age")), ENiagaraDataType::Vector));
+		if (ParticleRelativeTimes)
+		{
+			while (ParticleIndex >= 0)
+			{
+				if (ParticleRelativeTimes[ParticleIndex].X > 1.0f)
+				{
+					if (bGenerateDeathEvents)
+					{
+						DeathEventGenerator.OnDeath(ParticleIndex);
+					}
+
+					check(CurNumParticles > ParticleIndex);
+					// Particle is dead, move one from the end here. 
+					MoveParticleToIndex(--CurNumParticles, ParticleIndex);
+
+					DebuggerHook_OnDeath(this, ParticleIndex, CurNumParticles);
+				}
+				--ParticleIndex;
+			}
+		}
+		Data.SetNumInstances(CurNumParticles);
+
+		if (bGenerateDeathEvents)
+		{
+			DeathEventGenerator.EndTrackingDeaths();
+		}
+
+		// check if the emitter has officially died
+		if (GetTickState() == NTS_Dieing && CurNumParticles == 0)
+		{
+			SetTickState(NTS_Dead);
+		}
+	}
+
+	//Swap all data set buffers before doing the main tick on any simulation.
+	for (TPair<FNiagaraDataSetID, FNiagaraDataSet*> SetPair : DataSetMap)
+	{
+		SetPair.Value->Tick();
+	}
+}
 
 void FNiagaraSimulation::Tick(float DeltaSeconds)
 {
@@ -120,18 +304,23 @@ void FNiagaraSimulation::Tick(float DeltaSeconds)
 
 	SimpleTimer TickTime;
 
-
-	check(Data.GetNumAttributes() > 0);
+	check(Data.GetNumVariables() > 0);
 	check(PinnedProps->SpawnScriptProps.Script);
 	check(PinnedProps->UpdateScriptProps.Script);
+	
+	//Handle Event Actions.
+	auto CallEventActions = [&](FNiagaraEventReceiverProperties& Receiver)
+	{
+		for (auto& Action : Receiver.EmitterActions)
+		{
+			if (Action)
+				Action->PerformAction(*this, Receiver);
+		}
+	};
+	for (auto& Receiver : Props->SpawnScriptProps.EventReceivers)	{ CallEventActions(Receiver); }
+	for (auto& Receiver : Props->UpdateScriptProps.EventReceivers)	{ CallEventActions(Receiver); }
 
-	// Cache the ComponentToWorld transform.
-//	CachedComponentToWorld = Component.GetComponentToWorld();
-
-	Data.SwapBuffers();
-	Data.SetNumParticles(Data.GetPrevNumParticles());
-
-	int32 OrigNumParticles = Data.GetNumParticles();
+	int32 OrigNumParticles = Data.GetNumInstances();
 	int32 NumToSpawn = 0;
 
 	// Figure out how many we will spawn.
@@ -147,117 +336,107 @@ void FNiagaraSimulation::Tick(float DeltaSeconds)
 	if (TickState==NTS_Running || TickState==NTS_Dieing)
 	{
 		SCOPE_CYCLE_COUNTER(STAT_NiagaraSimulate);
-		RunVMScript(PinnedProps->UpdateScriptProps.Script, EUnusedAttributeBehaviour::Copy);
+		RunVMScript(PinnedProps->UpdateScriptProps, EUnusedAttributeBehaviour::Copy);
 	}
-	
+
+	DebuggerHook_PreSpawn(this, OrigNumParticles, NumToSpawn);
+
 	//Init new particles with the spawn script.
 	if (TickState==NTS_Running)
 	{
 		SCOPE_CYCLE_COUNTER(STAT_NiagaraSpawn);
-		Data.SetNumParticles(MaxNewParticles);
+		Data.SetNumInstances(MaxNewParticles);
 		//For now, zero any unused attributes here. But as this is really uninitialized data we should maybe make this a more serious error.
-		RunVMScript(PinnedProps->SpawnScriptProps.Script, EUnusedAttributeBehaviour::Zero, OrigNumParticles, NumToSpawn);
+		RunVMScript(PinnedProps->SpawnScriptProps, EUnusedAttributeBehaviour::Zero, OrigNumParticles, NumToSpawn);
 	}
 
-	// Iterate over looking for dead particles and move from the end of the list to the dead location, compacting in the process
+	if (bGenerateSpawnEvents)
 	{
-		SCOPE_CYCLE_COUNTER(STAT_NiagaraKill);
-		int32 CurNumParticles = OrigNumParticles = Data.GetNumParticles();
-		int32 ParticleIndex = 0;
-		const FVector4* ParticleRelativeTimes = Data.GetAttributeData(FNiagaraVariableInfo(FName(TEXT("Age")), ENiagaraDataType::Vector));
-		if (ParticleRelativeTimes)
-		{
-			while (ParticleIndex < OrigNumParticles)
-			{
-				if (ParticleRelativeTimes[ParticleIndex].X > 1.0f)
-				{
-					// Particle is dead, move one from the end here. 
-					MoveParticleToIndex(--CurNumParticles, ParticleIndex);
-				}
-				ParticleIndex++;
-			}
-		}
-		Data.SetNumParticles(CurNumParticles);
-
-		// check if the emitter has officially died
-		if (GetTickState() == NTS_Dieing && CurNumParticles == 0)
-		{
-			SetTickState(NTS_Dead);
-		}
+		SpawnEventGenerator.OnSpawned(OrigNumParticles, NumToSpawn);
 	}
 
 
-
 	CPUTimeMS = TickTime.GetElapsedMilliseconds();
 
-	INC_DWORD_STAT_BY(STAT_NiagaraNumParticles, Data.GetNumParticles());
+	INC_DWORD_STAT_BY(STAT_NiagaraNumParticles, Data.GetNumInstances());
 
 }
 
-void FNiagaraSimulation::RunVMScript(UNiagaraScript* Script, EUnusedAttributeBehaviour UnusedAttribBehaviour)
+void FNiagaraSimulation::RunVMScript(FNiagaraEmitterScriptProperties& ScriptProps, EUnusedAttributeBehaviour UnusedAttribBehaviour)
 {
-	RunVMScript(Script, UnusedAttribBehaviour, 0, Data.GetNumParticles());
+	RunVMScript(ScriptProps, UnusedAttribBehaviour, 0, Data.GetNumInstances());
 }
 
-void FNiagaraSimulation::RunVMScript(UNiagaraScript* Script, EUnusedAttributeBehaviour UnusedAttribBehaviour, uint32 StartParticle)
+void FNiagaraSimulation::RunVMScript(FNiagaraEmitterScriptProperties& ScriptProps, EUnusedAttributeBehaviour UnusedAttribBehaviour, uint32 StartParticle)
 {
-	RunVMScript(Script, UnusedAttribBehaviour, StartParticle, Data.GetNumParticles() - StartParticle);
+	RunVMScript(ScriptProps, UnusedAttribBehaviour, StartParticle, Data.GetNumInstances() - StartParticle);
 }
 
-void FNiagaraSimulation::RunVMScript(UNiagaraScript* Script, EUnusedAttributeBehaviour UnusedAttribBehaviour, uint32 StartParticle, uint32 NumParticles)
+void FNiagaraSimulation::RunVMScript(FNiagaraEmitterScriptProperties& ScriptProps, EUnusedAttributeBehaviour UnusedAttribBehaviour, uint32 StartParticle, uint32 NumParticles)
 {
 	if (NumParticles == 0)
 		return;
+	
+	UNiagaraScript* Script = ScriptProps.Script;
 
 	check(Script);
 	check(Script->ByteCode.Num() > 0);
 	check(Script->Attributes.Num() > 0);
-	check(Data.GetNumAttributes() > 0);
+	check(Data.GetNumVariables() > 0);
 
-	uint32 CurrNumParticles = Data.GetNumParticles();
+	uint32 CurrNumParticles = Data.GetNumInstances();
 
 	check(StartParticle + NumParticles <= CurrNumParticles);
 
 	const FVector4* InputBuffer = Data.GetPreviousBuffer();
 	const FVector4* OutputBuffer = Data.GetCurrentBuffer();
 
-	uint32 InputAttributeStride = Data.GetPrevParticleAllocation();
-	uint32 OutputAttributeStride = Data.GetParticleAllocation();
+	uint32 InputAttributeStride = Data.GetPrevDataAllocation();
+	uint32 OutputAttributeStride = Data.GetDataAllocation();
 
 	VectorRegister* InputRegisters[VectorVM::MaxInputRegisters] = { 0 };
 	VectorRegister* OutputRegisters[VectorVM::MaxOutputRegisters] = { 0 };
 	
 	//The current script run will be using NumScriptAttrs. We must pick out the Attributes form the simulation data in the order that they appear in the script.
 	//The remaining attributes being handled according to UnusedAttribBehaviour
-	const int32 NumSimulationAttrs = Data.GetNumAttributes();
+	const int32 NumSimulationAttrs = Data.GetNumVariables();
 	const int32 NumScriptAttrs = Script->Attributes.Num();
 
 	check(NumScriptAttrs < VectorVM::MaxInputRegisters);
 	check(NumScriptAttrs < VectorVM::MaxOutputRegisters);
 
+	int32 NumInputRegisters = 0;
+	int32 NumOutputRegitsers = 0;
 	// Setup input and output registers.
-	for (TPair<FNiagaraVariableInfo, uint32> AttrIndexPair : Data.GetAttributes())
+	for (TPair<FNiagaraVariableInfo, uint32> AttrIndexPair : Data.GetVariables())
 	{
 		const FNiagaraVariableInfo& AttrInfo = AttrIndexPair.Key;
 
 		FVector4* InBuff;
 		FVector4* OutBuff;
-		Data.GetAttributeData(AttrInfo, InBuff, OutBuff, StartParticle);
+		Data.GetVariableData(AttrInfo, InBuff, OutBuff, StartParticle);
 
 		int32 AttrIndex = Script->Attributes.Find(AttrInfo);
 		if (AttrIndex != INDEX_NONE)
 		{
+			//The input buffer can be NULL but only if this script never accesses it.
+			check(InBuff || !Script->Usage.bReadsAttriubteData);
+
 			//This attribute is used in this script so set it to the correct Input and Output buffer locations.
 			InputRegisters[AttrIndex] = (VectorRegister*)InBuff;
 			OutputRegisters[AttrIndex] = (VectorRegister*)OutBuff;
+			
+			NumInputRegisters += InBuff ? 1 : 0;
+			NumOutputRegitsers += OutBuff ? 1 : 0;
 		}	
 		else
 		{
 			//This attribute is not used in this script so handle it accordingly.
-			check(OutBuff != NULL);
+			check(OutBuff);
+
 			if (UnusedAttribBehaviour == EUnusedAttributeBehaviour::Copy)
 			{
-				check(InBuff != NULL);
+				check(InBuff);
 				FMemory::Memcpy(OutBuff, InBuff, NumParticles * sizeof(FVector4));
 			}
 			else if (UnusedAttribBehaviour == EUnusedAttributeBehaviour::Zero)
@@ -276,16 +455,81 @@ void FNiagaraSimulation::RunVMScript(UNiagaraScript* Script, EUnusedAttributeBeh
 	TArray<UNiagaraDataObject *>DataObjTable;
 	Script->ConstantData.FillConstantTable(Constants, ConstantTable, DataObjTable);
 
+
+	//Fill in the shared data table
+	TArray<FVectorVMSharedDataView, TInlineAllocator<64>> SharedDataTable;
+	
+	for (FNiagaraEventReceiverProperties Receiver : ScriptProps.EventReceivers)
+	{
+		//Look up to the effect for receiver sets as they can come from other emitters.
+		FNiagaraDataSetID ReceiverID = FNiagaraDataSetID(Receiver.Name, ENiagaraDataSetType::Event);
+		FNiagaraDataSetID SourceID = FNiagaraDataSetID(Receiver.SourceEventGenerator, ENiagaraDataSetType::Event);
+		
+		FNiagaraDataSet* SourceSet = ParentEffectInstance->GetDataSet(SourceID, Receiver.SourceEmitter);
+		FNiagaraDataSetProperties* ScriptEventReceiver = Script->EventReceivers.FindByPredicate([&](const FNiagaraDataSetProperties& R){ return ReceiverID == R.ID; });
+
+		check(ScriptEventReceiver);
+
+		uint32 NumInstances = SourceSet ? SourceSet->GetPrevNumInstances() : 0;
+		int32 ViewIdx = SharedDataTable.Add(FVectorVMSharedDataView(NumInstances, NumInstances - 1));
+		FVectorVMSharedDataView& View = SharedDataTable[ViewIdx];
+		for (FNiagaraVariableInfo& Var : ScriptEventReceiver->Variables)
+		{
+			FVector4* ReadData = SourceSet ? SourceSet->GetPrevVariableData(Var) : nullptr;
+			View.AddBuffer(ReadData);			
+		}
+	}
+
+	//Keep track of all the generator sets so that we can set their instance numbers after the VM run.
+	int32 StartOfGeneratorSets = SharedDataTable.Num();
+	TArray<FNiagaraDataSet*, TInlineAllocator<64>> GeneratorSets;
+
+	//TODO: Allow growing arrays for shared data / events in VM? HMMMM
+	for (FNiagaraEventGeneratorProperties Generator : ScriptProps.EventGenerators)
+	{
+		FNiagaraDataSet* GeneratorSet = DataSetMap.FindChecked(Generator.ID);
+		GeneratorSets.Add(GeneratorSet);
+	
+		check(GeneratorSet);
+
+		uint32 Allocation = GeneratorSet->GetDataAllocation();
+		uint32 NumInstances = GeneratorSet->GetNumInstances();
+
+		int32 ViewIdx = SharedDataTable.Add(FVectorVMSharedDataView(Allocation, NumInstances));
+		FVectorVMSharedDataView& View = SharedDataTable[ViewIdx];
+
+		for (FNiagaraVariableInfo& Var : Generator.Variables)
+		{
+			FVector4* WriteData = GeneratorSet->GetVariableData(Var);
+			View.AddBuffer(WriteData);
+		}
+	}
+	
+	DebuggerHook_PreScriptRun(this, ScriptProps.Script, Data.GetDataAllocation(), NumParticles, StartParticle);
+
 	VectorVM::Exec(
 		Script->ByteCode.GetData(),
 		InputRegisters,
-		NumScriptAttrs,
+		NumInputRegisters,
 		OutputRegisters,
-		NumScriptAttrs,
+		NumOutputRegitsers,
 		ConstantTable.GetData(),
 		DataObjTable.GetData(),
+		SharedDataTable.Num() > 0 ? SharedDataTable.GetData() : NULL,
 		NumParticles
 		);
+
+	//Go through our generators and set the num instances that were actually written ready for them to be read by receivers.
+	for (int32 i = 0; i < GeneratorSets.Num(); ++i)
+	{
+		FVectorVMSharedDataView& View = SharedDataTable[StartOfGeneratorSets + i];
+		FNiagaraDataSet* GeneratorSet = GeneratorSets[i];
+		check(GeneratorSet);
+		GeneratorSet->SetNumInstances(View.GetCounter());
+		INC_DWORD_STAT_BY(STAT_NiagaraNumCustomEvents, View.GetCounter());
+	}
+
+	DebuggerHook_PostScriptRun(this, Script);
 }
 
 bool FNiagaraSimulation::CheckAttriubtesForRenderer()
@@ -348,46 +592,3 @@ void FNiagaraSimulation::SetRenderModuleType(EEmitterRenderModuleType Type, ERHI
 		CheckAttriubtesForRenderer();
 	}
 }
-
-
-//////////////////////////////////////////////////////////////////////////
-
-UNiagaraEmitterProperties::UNiagaraEmitterProperties(const FObjectInitializer& Initiilaizer)
-: Super(Initiilaizer)
-, EmitterName(TEXT("New Emitter"))
-, bIsEnabled(true)
-, SpawnRate(50)
-, Material(nullptr)
-, RenderModuleType(RMT_Sprites)
-, StartTime(0.0f)
-, EndTime(0.0f)
-, RendererProperties(nullptr)
-, NumLoops(0)
-{
-}
-
-#if WITH_EDITOR
-void UNiagaraEmitterProperties::PostEditChangeProperty(struct FPropertyChangedEvent& PropertyChangedEvent)
-{
-	Super::PostEditChangeProperty(PropertyChangedEvent);
-	Init();
-}
-#endif
-
-void UNiagaraEmitterProperties::InitFromOldStruct(FDeprecatedNiagaraEmitterProperties& OldStruct)
-{
-	EmitterName = OldStruct.Name;
-	bIsEnabled = OldStruct.bIsEnabled;
-	SpawnRate = OldStruct.SpawnRate;
-	UpdateScriptProps.Script = OldStruct.UpdateScript;
-	SpawnScriptProps.Script = OldStruct.SpawnScript;
-	Material = OldStruct.Material;
-	RenderModuleType = OldStruct.RenderModuleType;
-	StartTime = OldStruct.StartTime;
-	EndTime = OldStruct.EndTime;
-	RendererProperties = OldStruct.RendererProperties;
-
-	UpdateScriptProps.ExternalConstants.InitFromOldMap(OldStruct.ExternalConstants);
-	SpawnScriptProps.ExternalConstants.InitFromOldMap(OldStruct.ExternalSpawnConstants);
-	NumLoops = OldStruct.NumLoops;
-}
\ No newline at end of file
diff --git a/Engine/Source/Runtime/Niagara/Private/NiagaraSimulationDebugger.cpp b/Engine/Source/Runtime/Niagara/Private/NiagaraSimulationDebugger.cpp
new file mode 100644
--- /dev/null
+++ b/Engine/Source/Runtime/Niagara/Private/NiagaraSimulationDebugger.cpp
@@ -0,0 +1,220 @@
+// Copyright 1998-2015 Epic Games, Inc. All Rights Reserved.
+
+#include "NiagaraPrivate.h"
+#include "NiagaraSimulation.h"
+#include "NiagaraEffectRenderer.h"
+#include "VectorVM.h"
+
+#if ENABLE_NIAGARA_SIMULATION_DEBUGGING
+
+FNiagaraSimulationDebugger& FNiagaraSimulationDebugger::Get()
+{
+	static FNiagaraSimulationDebugger Instance;
+	return Instance;
+}
+
+FNiagaraSimulationDebugger::FNiagaraSimulationDebugger()
+: DebugSim(NULL)
+, NumParticlesToDebugNextSpawn(0)
+, NumSpawnsToSkip(0)
+{
+}
+
+void FNiagaraSimulationDebugger::DebugNextNParticles(int32 NumParticlesToDebug, int32 SpawnsToSkip)
+{
+	NumParticlesToDebugNextSpawn = NumParticlesToDebug;
+	NumSpawnsToSkip = SpawnsToSkip;
+}
+
+void FNiagaraSimulationDebugger::InitForDebugging(FNiagaraSimulation* Sim, FString InDebugSimName)
+{
+	EndDebugging();
+
+	DebugSim = Sim;
+	DebugSimName = InDebugSimName;
+	UE_LOG(LogNiagara, Log, TEXT("======================================"));
+	UE_LOG(LogNiagara, Log, TEXT("Begun Debugging %s"), *DebugSimName);
+	UE_LOG(LogNiagara, Log, TEXT("======================================"));
+}
+
+void FNiagaraSimulationDebugger::EndDebugging()
+{
+	if (DebugSim)
+	{
+		UE_LOG(LogNiagara, Log, TEXT("======================================"));
+		UE_LOG(LogNiagara, Log, TEXT("Stopped Debugging %s"), *DebugSimName);
+		UE_LOG(LogNiagara, Log, TEXT("======================================"));
+		DebugSim = NULL;
+		DebugSimName.Empty();
+		CurrDebugInstances.Empty();
+	}
+}
+
+void FNiagaraSimulationDebugger::PreSpawn(FNiagaraSimulation* Sim, int32 CurrNumParticles, int32 NumToSpawn)
+{
+	if (DebugSim == Sim)
+	{
+		int32 CurrParticleToDebug = CurrNumParticles;
+		while (--NumToSpawn >= 0 && NumSpawnsToSkip-- == 0 && NumParticlesToDebugNextSpawn-- > 0)
+		{
+			UE_LOG(LogNiagara, Log, TEXT("Begun Debugging Particle #%d"), CurrParticleToDebug);
+			CurrDebugInstances.Add(CurrParticleToDebug++);
+		}
+	}
+}
+
+void FNiagaraSimulationDebugger::OnDeath(FNiagaraSimulation* Sim, int32 DyingParticle, int32 ReplacementParticle)
+{
+	if (IsDebugging(Sim))
+	{
+		if (CurrDebugInstances.Remove(DyingParticle) > 0)
+		{
+			UE_LOG(LogNiagara, Log, TEXT("PARTICLE #%d DIED!"), DyingParticle);
+		}
+		if (CurrDebugInstances.Remove(ReplacementParticle) > 0)
+		{
+			//If we're debugging the particle being moved into the dead particles position then we need to remove that and begin tracking it's new index.
+			UE_LOG(LogNiagara, Log, TEXT("PARTICLE MOVE! Stopped debugging %d. Now debugging #%d"), ReplacementParticle, DyingParticle);
+			CurrDebugInstances.Add(DyingParticle);
+		}
+	}
+}
+
+void FNiagaraSimulationDebugger::PreScriptRun(FNiagaraSimulation* Sim, UNiagaraScript* Script, int32 TotalNumParticles, int32 NumParticlesProcessed, int32 StartingParticle)
+{
+	if (IsDebugging(Sim))
+	{
+		UE_LOG(LogNiagara, Log, TEXT("=== %s ==================================="), *Script->GetName());
+		UE_LOG(LogNiagara, Log, TEXT("Particles: %d - ToProcess: %d - Start: %d"), TotalNumParticles, NumParticlesProcessed, StartingParticle);
+
+		VMDebugger.InitForScriptRun(StartingParticle, CurrDebugInstances);
+		VectorVM::AttachDebugger(&VMDebugger);
+	}
+}
+
+FString FNiagaraSimulationDebugger::GetValueString(FVector4 Value)
+{
+	return FString::Printf(TEXT("{%.5f, %.5f, %.5f, %.5f}"), Value.X, Value.Y, Value.Z, Value.W);
+}
+
+FString FNiagaraSimulationDebugger::GetOpArgString(const VectorVM::FOpDebugInfo& DebugInfo, int32 ArgIndex)
+{
+	const VectorVM::FDataLocationInfo& LocInfo = DebugInfo.LocationInfo[ArgIndex];
+	FString PreOpStr, PostOpStr;
+	// 		if (Type == VectorVM::EVMType::float)
+	// 		{
+	// 		}
+	// 		else if (Type == VectorVM::EVMType::int)
+	// 		{
+	// 		}
+	// 		else 
+	if (DebugInfo.OpType == VectorVM::EVMType::Vector4)
+	{
+		PreOpStr = GetValueString(DebugInfo.PreOpValues[ArgIndex].v);
+		PostOpStr = GetValueString(DebugInfo.PostOpValues[ArgIndex].v);
+	}
+	FString LocNameStr = VectorVM::GetOperandLocationName(LocInfo.Location);
+	FString LocStr = FString::Printf(TEXT("%s%d"), *LocNameStr, LocInfo.PrimaryLocationIndex);
+	if (LocInfo.SecondaryLocationIndex != INDEX_NONE)
+		LocStr.Append(FString::Printf(TEXT("-%d"), LocInfo.SecondaryLocationIndex));
+	if (LocInfo.TertiaryLocationIndex != INDEX_NONE)
+		LocStr.Append(FString::Printf(TEXT("-%d"), LocInfo.TertiaryLocationIndex));
+
+	return FString::Printf(TEXT("[%s-%p] %s --> %s"), *LocStr, LocInfo.MemoryAddress, *PreOpStr, *PostOpStr);
+}
+
+void FNiagaraSimulationDebugger::PostScriptRun(FNiagaraSimulation* Sim, UNiagaraScript* Script)
+{
+	if (IsDebugging(Sim))
+	{
+		VectorVM::DetachDebugger(&VMDebugger);
+		for (int32 Particle : CurrDebugInstances)
+		{
+			FString PrefixStr = FString::Printf(TEXT("[%s|%d]\t"), *(Script->GetName()), Particle);
+			if (const TArray<VectorVM::FOpDebugInfo>* InfoArray = VMDebugger.GetDebugInfoForInstance(Particle))
+			{
+				for (const VectorVM::FOpDebugInfo& DebugInfo : *InfoArray)
+				{
+					FString OpHeader = *VectorVM::GetOpName(DebugInfo.Op);
+					if (DebugInfo.Op == EVectorVMOp::output)
+					{
+						const VectorVM::FDataLocationInfo& DestLocInfo = DebugInfo.LocationInfo[(int32)VectorVM::EKernelArgs::Dest];
+						if (DestLocInfo.PrimaryLocationIndex < Script->Attributes.Num() && DestLocInfo.PrimaryLocationIndex >= 0)
+						{
+							OpHeader.Append(FString::Printf(TEXT(" [Attrib: %s]"), *Script->Attributes[DestLocInfo.PrimaryLocationIndex].Name.ToString()));
+						}
+					}
+					UE_LOG(LogNiagara, Log, TEXT("%s[%s] ----------------------------------"), *PrefixStr, *OpHeader);
+					UE_LOG(LogNiagara, Log, TEXT("%sDest%s"), *PrefixStr, *GetOpArgString(DebugInfo, (int32)VectorVM::EKernelArgs::Dest));
+					for (int32 Arg = 0; Arg < DebugInfo.NumArgs; ++Arg)
+					{
+						UE_LOG(LogNiagara, Log, TEXT("%sArg%d%s"), *PrefixStr, Arg, *GetOpArgString(DebugInfo, Arg + 1));
+					}
+				}
+			}
+			else
+			{
+				UE_LOG(LogNiagara, Log, TEXT("%s FAILED TO FIND DEBUG INFO FOR PARTICLE %d!"), *PrefixStr, Particle);
+			}
+		}
+		UE_LOG(LogNiagara, Log, TEXT("======================================"));
+	}
+}
+
+static void	BeginDebuggingNiagaraSimulation(const TArray<FString>& Args, UWorld* InWorld)
+{
+	if (Args.Num() < 2)
+		return;
+
+	FString ActorName = Args[0];
+	FString EmitterName = Args[1];
+	for (TActorIterator<AActor> It(InWorld, ANiagaraActor::StaticClass()); It; ++It)
+	{
+		if (It->GetName() == Args[0])
+		{
+			if (UNiagaraComponent* Comp = Cast<UNiagaraComponent>(It->FindComponentByClass(UNiagaraComponent::StaticClass())))
+			{
+				if (FNiagaraEffectInstance* EffectInst = Comp->GetEffectInstance().Get())
+				{
+					if (FNiagaraSimulation* Sim = EffectInst->GetEmitter(EmitterName))
+					{
+						FNiagaraSimulationDebugger::Get().InitForDebugging(Sim, EmitterName);
+					}
+				}
+			}
+		}
+	}
+}
+
+FAutoConsoleCommandWithWorldAndArgs GBeginDebuggingSimulation(
+	TEXT("Niagara.BeginDebuggingSimulation"),
+	TEXT(""),
+	FConsoleCommandWithWorldAndArgsDelegate::CreateStatic(BeginDebuggingNiagaraSimulation)
+	);
+
+void DebugNextNParticlesToSpawn(const TArray<FString>& Args)
+{
+	if (Args.Num() >= 1)
+	{
+		int32 N = 0;
+		int32 ToSkip = 0;
+		LexicalConversion::FromString(N, *Args[0]);
+		if (Args.Num() > 1)
+			LexicalConversion::FromString(ToSkip, *Args[1]);
+
+		FNiagaraSimulationDebugger::Get().DebugNextNParticles(N, ToSkip);
+	}
+}
+
+FAutoConsoleCommand GDebugNextNParticlesToSpawn(
+	TEXT("Niagara.DebugNextNParticlesToSpawn"),
+	TEXT(""),
+	FConsoleCommandWithArgsDelegate::CreateStatic(DebugNextNParticlesToSpawn)
+	);
+
+#endif
+
+FString LexicalConversion::ToString(VectorRegister Register)
+{
+	return FString::Printf(TEXT("{%.4f, %.4f, %.4f, %.4f}"), VectorGetComponent(Register, 0), VectorGetComponent(Register, 1), VectorGetComponent(Register, 2), VectorGetComponent(Register, 3));
+}
diff --git a/Engine/Source/Runtime/Niagara/Private/NiagaraSimulationDebugger.h b/Engine/Source/Runtime/Niagara/Private/NiagaraSimulationDebugger.h
new file mode 100644
--- /dev/null
+++ b/Engine/Source/Runtime/Niagara/Private/NiagaraSimulationDebugger.h
@@ -0,0 +1,49 @@
+// Copyright 1998-2015 Epic Games, Inc. All Rights Reserved.
+
+#pragma once
+
+#include "NiagaraScript.h"
+
+//#define ENABLE_NIAGARA_SIMULATION_DEBUGGING (0)
+#define ENABLE_NIAGARA_SIMULATION_DEBUGGING (1) && WITH_EDITOR && ENABLE_VM_DEBUGGING
+
+namespace LexicalConversion
+{
+	FString ToString(VectorRegister Register);
+};
+
+#if ENABLE_NIAGARA_SIMULATION_DEBUGGING
+
+class FNiagaraSimulationDebugger
+{
+	VectorVM::FVectorVMDebugger VMDebugger;
+	FNiagaraSimulation* DebugSim;
+	FString DebugSimName;
+
+	TArray<int32> CurrDebugInstances;
+
+	int32 NumParticlesToDebugNextSpawn;
+	int32 NumSpawnsToSkip;
+
+	FString GetValueString(FVector4 Value);
+	FString GetOpArgString(const VectorVM::FOpDebugInfo& DebugInfo, int32 ArgIndex);
+public:
+
+	static FNiagaraSimulationDebugger& Get();
+
+	FNiagaraSimulationDebugger();
+
+	FORCEINLINE bool IsDebugging(FNiagaraSimulation* InSim){ return DebugSim == InSim && CurrDebugInstances.Num() > 0; }
+
+	void DebugNextNParticles(int32 NumParticlesToDebug, int32 SpawnsToSkip);
+	void InitForDebugging(FNiagaraSimulation* Sim, FString InDebugSimName);
+	void EndDebugging();
+
+	// Hooks for simulations
+	void PreSpawn(FNiagaraSimulation* Sim, int32 CurrNumParticles, int32 NumToSpawn);
+	void OnDeath(FNiagaraSimulation* Sim, int32 DyingParticle, int32 ReplacementParticle);
+	void PreScriptRun(FNiagaraSimulation* Sim, UNiagaraScript* Script, int32 TotalNumParticles, int32 NumParticlesProcessed, int32 StartingParticle);
+	void PostScriptRun(FNiagaraSimulation* Sim, UNiagaraScript* Script);
+};
+
+#endif
diff --git a/Engine/Source/Runtime/Niagara/Public/NiagaraCommon.h b/Engine/Source/Runtime/Niagara/Public/NiagaraCommon.h
--- a/Engine/Source/Runtime/Niagara/Public/NiagaraCommon.h
+++ b/Engine/Source/Runtime/Niagara/Public/NiagaraCommon.h
@@ -69,3 +69,60 @@ FORCEINLINE uint32 GetTypeHash(const FNiagaraVariableInfo& Var)
 {
 	return HashCombine(GetTypeHash(Var.Name), (uint32)(Var.Type.GetValue()));
 }
+
+
+UENUM()
+enum class ENiagaraDataSetType : uint8
+{
+	ParticleData,
+	Shared,
+	Event,
+};
+
+USTRUCT()
+struct FNiagaraDataSetID
+{
+	GENERATED_USTRUCT_BODY()
+
+	FNiagaraDataSetID()
+	: Name(NAME_None)
+	, Type(ENiagaraDataSetType::Event)
+	{}
+
+	FNiagaraDataSetID(FName InName, ENiagaraDataSetType InType)
+		: Name(InName)
+		, Type(InType)
+	{}
+
+	UPROPERTY(EditAnywhere, Category = "Data Set")
+	FName Name;
+
+	UPROPERTY()
+	TEnumAsByte<ENiagaraDataSetType> Type;
+
+	FORCEINLINE bool operator==(const FNiagaraDataSetID& Other)const
+	{
+		return Name == Other.Name && Type == Other.Type;
+	}
+
+	FORCEINLINE bool operator!=(const FNiagaraDataSetID& Other)const
+	{
+		return !(*this == Other);
+	}
+
+	//Some special case data sets.
+	static FNiagaraDataSetID DeathEvent;
+	static FNiagaraDataSetID SpawnEvent;
+	static FNiagaraDataSetID CollisionEvent;
+};
+
+FORCEINLINE FArchive& operator<<(FArchive& Ar, FNiagaraDataSetID& VarInfo)
+{
+	Ar << VarInfo.Name << VarInfo.Type;
+	return Ar;
+}
+
+FORCEINLINE uint32 GetTypeHash(const FNiagaraDataSetID& Var)
+{
+	return HashCombine(GetTypeHash(Var.Name), (uint32)(Var.Type.GetValue()));
+}
diff --git a/Engine/Source/Runtime/Niagara/Public/NiagaraComponent.h b/Engine/Source/Runtime/Niagara/Public/NiagaraComponent.h
--- a/Engine/Source/Runtime/Niagara/Public/NiagaraComponent.h
+++ b/Engine/Source/Runtime/Niagara/Public/NiagaraComponent.h
@@ -6,112 +6,8 @@
 
 #include "NiagaraComponent.generated.h"
 
-/**
-* Contains the simulation data for a single FNiagaraSimulation
-* all buffers for particle attributes are stored here for convenient access
-*/
-class FNiagaraEmitterParticleData
-{
-public:
-	FNiagaraEmitterParticleData()
-		: CurrentBuffer(0)
-	{
-		Reset();
-	}
-
-	~FNiagaraEmitterParticleData() {}
-
-	void Allocate(uint32 NumExpectedParticles)
-	{
-		ParticleBuffers[CurrentBuffer].Reset(NumExpectedParticles * AttrMap.Num());
-		ParticleBuffers[CurrentBuffer].AddUninitialized(NumExpectedParticles * AttrMap.Num());
-
-		ParticleAllocation[CurrentBuffer] = NumExpectedParticles;
-	}
-
-	void Reset()
-	{
-		AttrMap.Empty();
-		NumParticles[0] = 0;
-		NumParticles[1] = 0;
-		ParticleAllocation[0] = 0;
-		ParticleAllocation[1] = 0;
-		CurrentBuffer = 0;
-	}
-
-	const FVector4* GetAttributeData(const FNiagaraVariableInfo& AttrID) const
-	{
-		const uint32* Offset = AttrMap.Find(AttrID);
-		return Offset ? ParticleBuffers[CurrentBuffer].GetData() + (*Offset * ParticleAllocation[CurrentBuffer]) : NULL;
-	}
-
-	void GetAttributeData(const FNiagaraVariableInfo& AttrID, FVector4*& PrevBuffer, FVector4*& CurrBuffer, int32 StartParticle) 
-	{
-		const uint32* Offset = AttrMap.Find(AttrID);
-		if (Offset)
-		{
-			PrevBuffer = ParticleBuffers[CurrentBuffer ^ 0x1].GetData() + (*Offset * ParticleAllocation[CurrentBuffer ^ 0x1]) + StartParticle;
-			CurrBuffer = ParticleBuffers[CurrentBuffer].GetData() + (*Offset * ParticleAllocation[CurrentBuffer]) + StartParticle;
-		}
-		else
-		{
-			PrevBuffer = NULL;
-			CurrBuffer = NULL;
-		}
-	}
-
-	bool HasAttriubte(const FNiagaraVariableInfo& AttrID)
-	{
-		return AttrMap.Find(AttrID) != NULL;
-	}
-
-	void SetAttributes(const TArray<FNiagaraVariableInfo>& Attrs)
-	{
-		Reset();
-		for (const FNiagaraVariableInfo& Attr : Attrs)
-		{
-			//TODO: Support attributes of any type. This may as well wait until the move of matrix ops etc over into UNiagaraScriptStructs.
-			uint32 Idx = AttrMap.Num();
-			AttrMap.Add(Attr, Idx);
-		}
-	}
-
-	void AddAttributes(const TArray<FNiagaraVariableInfo>& Attrs)
-	{
-		for (const FNiagaraVariableInfo& Attr : Attrs)
-		{
-			//TODO: Support attributes of any type. This may as well wait until the move of matrix ops etc over into UNiagaraScriptStructs.
-			uint32 Idx = AttrMap.Num();
-			if (!AttrMap.Find(Attr))
-			{
-				AttrMap.Add(Attr, Idx);
-			}
-		}
-	}
-
-	const TMap<FNiagaraVariableInfo, uint32>& GetAttributes()const { return AttrMap; }
-	int GetNumAttributes()				{ return AttrMap.Num(); }
-
-	uint32 GetNumParticles() const		{ return NumParticles[CurrentBuffer]; }
-	uint32 GetPrevNumParticles() const		{ return NumParticles[CurrentBuffer ^ 0x1]; }
-	uint32 GetParticleAllocation() const { return ParticleAllocation[CurrentBuffer]; }
-	uint32 GetPrevParticleAllocation() const { return ParticleAllocation[CurrentBuffer ^ 0x1]; }
-	void SetNumParticles(uint32 Num)	{ NumParticles[CurrentBuffer] = Num; }
-	void SwapBuffers()					{ CurrentBuffer ^= 0x1; }
-
-	FVector4* GetCurrentBuffer()		{ return ParticleBuffers[CurrentBuffer].GetData(); }
-	FVector4* GetPreviousBuffer()		{ return ParticleBuffers[CurrentBuffer^0x1].GetData(); }
-
-	int GetBytesUsed()	{ return (ParticleBuffers[0].Num() + ParticleBuffers[1].Num()) * 16 + AttrMap.Num() * 4; }
-private:
-	uint32 CurrentBuffer;
-	uint32 NumParticles[2];
-	uint32 ParticleAllocation[2];
-	TArray<FVector4> ParticleBuffers[2];
-	TMap<FNiagaraVariableInfo, uint32> AttrMap;
-};
-
-
+class UNiagaraEffect;
+class FNiagaraEffectInstance;
 
 /**
 * UNiagaraComponent is the primitive component for a Niagara effect.
@@ -124,9 +20,9 @@ class NIAGARA_API UNiagaraComponent : public UPrimitiveComponent
 	GENERATED_UCLASS_BODY()
 private:
 	UPROPERTY()
-	class UNiagaraEffect* Asset;
+	UNiagaraEffect* Asset;
 
-	class FNiagaraEffectInstance* EffectInstance;
+	TSharedPtr<FNiagaraEffectInstance> EffectInstance;
 
 	// Begin UActorComponent interface.
 protected:
@@ -147,8 +43,8 @@ public:
 	void SetAsset(UNiagaraEffect* InAsset);
 	UNiagaraEffect* GetAsset() const { return Asset; }
 
-	FNiagaraEffectInstance* GetEffectInstance()	const { return EffectInstance; }
-	void SetEffectInstance(FNiagaraEffectInstance* InInstance)	{ EffectInstance = InInstance; }
+	TSharedPtr<FNiagaraEffectInstance> GetEffectInstance()	const { return EffectInstance; }
+	void SetEffectInstance(TSharedPtr<FNiagaraEffectInstance> InInstance)	{ EffectInstance = InInstance; }
 
 	// Begin UObject interface.
 #if WITH_EDITOR
