Copying //UE4/Dev-Rendering to //UE4/Dev-Main (Source: //UE4/Dev-Rendering @ 3045398)

diff --git a/Engine/Source/Runtime/Engine/Private/Streaming/TextureInstanceManager.cpp b/Engine/Source/Runtime/Engine/Private/Streaming/TextureInstanceManager.cpp
--- a/Engine/Source/Runtime/Engine/Private/Streaming/TextureInstanceManager.cpp
+++ b/Engine/Source/Runtime/Engine/Private/Streaming/TextureInstanceManager.cpp
@@ -702,7 +702,7 @@ TRefCountPtr<const FTextureInstanceState> FDynamicTextureInstanceManager::GetAsy
 	return TRefCountPtr<const FTextureInstanceState>(AsyncState);
 }
 
-void FTextureInstanceAsyncView::UpdateBoundSizes_Async(const TArray<FStreamingViewInfo>& ViewInfos, float LastUpdateTime, bool bUseNewMetrics, float MaxEffectiveScreenSize)
+void FTextureInstanceAsyncView::UpdateBoundSizes_Async(const TArray<FStreamingViewInfo>& ViewInfos, float LastUpdateTime, const FTextureStreamingSettings& Settings)
 {
 	if (!State.IsValid())  return;
 
@@ -742,7 +742,7 @@ void FTextureInstanceAsyncView::UpdateBoundSizes_Async(const TArray<FStreamingVi
 		{
 			const FStreamingViewInfo& ViewInfo = ViewInfos[ViewIndex];
 
-			const float EffectiveScreenSize = (MaxEffectiveScreenSize > 0.0f) ? FMath::Min(MaxEffectiveScreenSize, ViewInfo.ScreenSize) : ViewInfo.ScreenSize;
+			const float EffectiveScreenSize = (Settings.MaxEffectiveScreenSize > 0.0f) ? FMath::Min(Settings.MaxEffectiveScreenSize, ViewInfo.ScreenSize) : ViewInfo.ScreenSize;
 			const float ScreenSizeFloat = EffectiveScreenSize * ViewInfo.BoostFactor * .5f; // Multiply by half since the ratio factors map to half the screen only
 
 			const VectorRegister ScreenSize = VectorLoadFloat1( &ScreenSizeFloat );
@@ -751,7 +751,7 @@ void FTextureInstanceAsyncView::UpdateBoundSizes_Async(const TArray<FStreamingVi
 			const VectorRegister ViewOriginZ = VectorLoadFloat1( &ViewInfo.ViewOrigin.Z );
 
 			VectorRegister DistSqMinusRadiusSq = VectorZero();
-			if (bUseNewMetrics)
+			if (Settings.bUseNewMetrics)
 			{
 				// In this case DistSqMinusRadiusSq will contain the distance to the box^2
 				VectorRegister Temp = VectorSubtract( ViewOriginX, OriginX );
diff --git a/Engine/Source/Runtime/Engine/Private/Streaming/TextureInstanceManager.h b/Engine/Source/Runtime/Engine/Private/Streaming/TextureInstanceManager.h
--- a/Engine/Source/Runtime/Engine/Private/Streaming/TextureInstanceManager.h
+++ b/Engine/Source/Runtime/Engine/Private/Streaming/TextureInstanceManager.h
@@ -303,7 +303,7 @@ public:
 
 	FTextureInstanceAsyncView(const TRefCountPtr<const FTextureInstanceState>& InState) : State(InState) {}
 
-	void UpdateBoundSizes_Async(const TArray<FStreamingViewInfo>& ViewInfos, float LastUpdateTime, bool bUseNewMetrics, float MaxEffectiveScreenSize);
+	void UpdateBoundSizes_Async(const TArray<FStreamingViewInfo>& ViewInfos, float LastUpdateTime, const FTextureStreamingSettings& Settings);
 
 	// MaxSize : Biggest texture size for all instances.
 	// MaxSize_VisibleOnly : Biggest texture size for visble instances only.
@@ -396,7 +396,6 @@ class FLevelTextureManager
 public:
 
 	FLevelTextureManager(ULevel* InLevel) : Level(InLevel), bHasTextures(false), bToDelete(false) {}
-	~FLevelTextureManager() { check(bToDelete); }
 
 
 	ULevel* GetLevel() const { return Level; }
diff --git a/Engine/Source/Runtime/Engine/Private/Streaming/TextureStreamingBuild.cpp b/Engine/Source/Runtime/Engine/Private/Streaming/TextureStreamingBuild.cpp
--- a/Engine/Source/Runtime/Engine/Private/Streaming/TextureStreamingBuild.cpp
+++ b/Engine/Source/Runtime/Engine/Private/Streaming/TextureStreamingBuild.cpp
@@ -25,7 +25,7 @@ bool WaitForShaderCompilation(const FText& Message, FSlowTask& BuildTextureStrea
 			const int32 NumberOfShadersCompiledThisFrame = RemainingShaders - RemainingShadersThisFrame;
 
 			SlowTask.EnterProgressFrame((float)NumberOfShadersCompiledThisFrame);
-			BuildTextureStreamingTask.EnterProgressFrame((float)NumberOfShadersCompiledThisFrame / (float)NumShadersToBeCompiled * (1.f - KINDA_SMALL_NUMBER));
+			BuildTextureStreamingTask.EnterProgressFrame((float)NumberOfShadersCompiledThisFrame / (float)NumShadersToBeCompiled);
 			if (GWarn->ReceivedUserCancel()) return false;
 
 			RemainingShaders = RemainingShadersThisFrame;
@@ -90,9 +90,10 @@ int32 GetNumPrimitiveComponents(UWorld* InWorld)
  * @param QualityLevel		The quality level for the shaders.
  * @param FeatureLevel		The feature level for the shaders.
  * @param TexCoordScales	The map to store material interfaces used by the current primitive world.
+ * @param bIncremental		true if the build should only update components that have no data.
  * @return true if the operation is a success, false if it was canceled.
  */
-ENGINE_API bool BuildTextureStreamingShaders(UWorld* InWorld, EMaterialQualityLevel::Type QualityLevel, ERHIFeatureLevel::Type FeatureLevel, OUT FTexCoordScaleMap& TexCoordScales, FSlowTask& BuildTextureStreamingTask)
+ENGINE_API bool BuildTextureStreamingShaders(UWorld* InWorld, EMaterialQualityLevel::Type QualityLevel, ERHIFeatureLevel::Type FeatureLevel, OUT FTexCoordScaleMap& TexCoordScales , bool bIncremental, FSlowTask& BuildTextureStreamingTask)
 {
 #if WITH_EDITORONLY_DATA
 
@@ -108,7 +109,14 @@ ENGINE_API bool BuildTextureStreamingShaders(UWorld* InWorld, EMaterialQualityLe
 		return false;
 	}
 
-	FDebugViewModeMaterialProxy::ClearAllShaders();
+	if (!bIncremental)
+	{
+		FDebugViewModeMaterialProxy::ClearAllShaders();
+	}
+	else // Otherwise, if there are no shaders, then we need to rebuild them all.
+	{
+		bIncremental = FDebugViewModeMaterialProxy::HasAnyShaders();
+	}
 
 	// Without the new metrics, the shaders will are not compiled, making the texcoord scale viewmode non functional.
 	if (bUseNewMetrics && AllowDebugViewPS(DVSM_MaterialTexCoordScalesAnalysis, GetFeatureLevelShaderPlatform(FeatureLevel)))
@@ -121,21 +129,32 @@ ENGINE_API bool BuildTextureStreamingShaders(UWorld* InWorld, EMaterialQualityLe
 			for (UPrimitiveComponent* Primitive : Components)
 			{
 				SlowTask.EnterProgressFrame();
-				BuildTextureStreamingTask.EnterProgressFrame((1.f - KINDA_SMALL_NUMBER) / (float)Components.Num());
+				BuildTextureStreamingTask.EnterProgressFrame(1.f / (float)Components.Num());
 				if (GWarn->ReceivedUserCancel())
 				{
 					FDebugViewModeMaterialProxy::ClearAllShaders();
 					return false;
 				}
 
+				// When only updating (viewmode path), skip primitives that already have data.
+				if (bIncremental && !Primitive->HasMissingStreamingSectionData(true))
+					continue;
+
 				TArray<UMaterialInterface*> Materials;
 				Primitive->GetUsedMaterials(Materials);
 
-				for (int32 MaterialIndex = 0; MaterialIndex < Materials.Num(); ++MaterialIndex)
+				for (UMaterialInterface* MaterialInterface : Materials)
 				{
-					UMaterialInterface* MaterialInterface = Materials[MaterialIndex];
 					if (!MaterialInterface) continue;
 
+					// Landscape material resources can not be used. See logic in FLandscapeMaterialResource::ShouldCache().
+					const FMaterial* Material = MaterialInterface->GetMaterialResource(FeatureLevel);
+					if (!Material || Material->IsUsedWithLandscape())
+					{
+						UE_LOG(LogLevel, Warning, TEXT("Landscape material %s not supported, skipping shader"), *MaterialInterface->GetName());
+						continue;
+					}
+
 					if (!TexCoordScales.Contains(MaterialInterface))
 					{
 						TexCoordScales.Add(MaterialInterface);
@@ -169,6 +188,57 @@ ENGINE_API bool BuildTextureStreamingShaders(UWorld* InWorld, EMaterialQualityLe
 #endif
 }
 
+ENGINE_API bool UpdateComponentStreamingSectionData(UWorld* InWorld, const FTexCoordScaleMap& InTexCoordScales, bool bIncremental, FSlowTask& BuildTextureStreamingTask)
+{
+#if WITH_EDITORONLY_DATA
+	if (!InWorld) return false;
+	const double StartTime = FPlatformTime::Seconds();
+
+	// ====================================================
+	// Build per primitive data
+	// ====================================================
+	const float NumPrimitiveComponents = (float)GetNumPrimitiveComponents(InWorld);
+	FScopedSlowTask SlowTask(NumPrimitiveComponents, (LOCTEXT("TextureStreamingBuild_ComponentDataUpdate", "Updating Component Data")));
+
+	for (int32 LevelIndex = 0; LevelIndex < InWorld->GetNumLevels(); LevelIndex++)
+	{
+		ULevel* Level = InWorld->GetLevel(LevelIndex);
+		if (!Level) continue;
+		TArray<UTexture2D*> LevelTextures;
+
+		for (AActor* Actor : Level->Actors)
+		{
+			if (!Actor) continue;
+			bool bComponentUpdated = false;
+
+			TInlineComponentArray<UPrimitiveComponent*> Primitives;
+			Actor->GetComponents<UPrimitiveComponent>(Primitives);
+			for (UPrimitiveComponent* Primitive : Primitives)
+			{
+				SlowTask.EnterProgressFrame();
+				BuildTextureStreamingTask.EnterProgressFrame(1.f / NumPrimitiveComponents);
+				if (GWarn->ReceivedUserCancel()) return false;
+
+				if (bIncremental && !Primitive->HasMissingStreamingSectionData(InTexCoordScales.Num() > 0))
+					continue;
+
+				Primitive->UpdateStreamingSectionData(InTexCoordScales);
+				bComponentUpdated = true;
+			}
+
+			if (bComponentUpdated)
+			{
+				Actor->MarkComponentsRenderStateDirty();
+			}
+		}
+	}
+	UE_LOG(LogLevel, Display, TEXT("Update Texture Streaming Data took %.3f seconds."), FPlatformTime::Seconds() - StartTime);
+	return true;
+#else
+	return false;
+#endif
+}
+
 ENGINE_API bool BuildTextureStreamingData(UWorld* InWorld, const FTexCoordScaleMap& InTexCoordScales, EMaterialQualityLevel::Type QualityLevel, ERHIFeatureLevel::Type FeatureLevel, FSlowTask& BuildTextureStreamingTask)
 {
 #if WITH_EDITORONLY_DATA
@@ -197,7 +267,7 @@ ENGINE_API bool BuildTextureStreamingData(UWorld* InWorld, const FTexCoordScaleM
 			for (UPrimitiveComponent* Primitive : Primitives)
 			{
 				SlowTask.EnterProgressFrame();
-				BuildTextureStreamingTask.EnterProgressFrame((1.f - KINDA_SMALL_NUMBER) / NumPrimitiveComponents);
+				BuildTextureStreamingTask.EnterProgressFrame(1.f / NumPrimitiveComponents);
 				if (GWarn->ReceivedUserCancel()) return false;
 
 				UStaticMeshComponent* StaticMeshComponent = Cast<UStaticMeshComponent>(Primitive);
diff --git a/Engine/Source/Runtime/Engine/Private/Streaming/TextureStreamingHelpers.cpp b/Engine/Source/Runtime/Engine/Private/Streaming/TextureStreamingHelpers.cpp
--- a/Engine/Source/Runtime/Engine/Private/Streaming/TextureStreamingHelpers.cpp
+++ b/Engine/Source/Runtime/Engine/Private/Streaming/TextureStreamingHelpers.cpp
@@ -134,11 +134,47 @@ TAutoConsoleVariable<float> CVarStreamingMipBias(
 	ECVF_Scalability
 	);
 
-TAutoConsoleVariable<int32> CVarOnlyStreamInTextures(
-	TEXT("r.OnlyStreamInTextures"),
+TAutoConsoleVariable<int32> CVarStreamingUsePerTextureBias(
+	TEXT("r.Streaming.UsePerTextureBias"),
 	0,
-	TEXT("If set to 1, texture will only be streamed in, not out"),
-	ECVF_RenderThreadSafe);
+	TEXT("If non-zero, each texture will be assigned a mip bias between 0 and MipBias as required to fit in budget."),
+	ECVF_Default);
+
+
+TAutoConsoleVariable<int32> CVarStreamingFullyLoadUsedTextures(
+	TEXT("r.Streaming.FullyLoadUsedTextures"),
+	0,
+	TEXT("If non-zero, all used texture will be fully streamed in as fast as possible"),
+	ECVF_Default);
+
+TAutoConsoleVariable<int32> CVarStreamingUseAllMips(
+	TEXT("r.Streaming.UseAllMips"),
+	0,
+	TEXT("If non-zero, all available mips will be used"),
+	ECVF_Default);
+
+
+void FTextureStreamingSettings::Update()
+{
+	MaxEffectiveScreenSize = CVarStreamingScreenSizeEffectiveMax.GetValueOnAnyThread();
+	MaxTempMemoryAllowed = CVarStreamingMaxTempMemoryAllowed.GetValueOnAnyThread();
+	DropMips = CVarStreamingDropMips.GetValueOnAnyThread();
+	HLODStrategy = CVarStreamingHLODStrategy.GetValueOnAnyThread();
+	HiddenPrimitiveScale = CVarStreamingHiddenPrimitiveScale.GetValueOnAnyThread();
+	MipBias = FMath::Max<float>(0.f, CVarStreamingMipBias.GetValueOnAnyThread());
+	PoolSize = CVarStreamingPoolSize.GetValueOnAnyThread();
+	bUsePerTextureBias = CVarStreamingUsePerTextureBias.GetValueOnAnyThread() != 0;
+	bUseNewMetrics = CVarStreamingUseNewMetrics.GetValueOnAnyThread() != 0;
+	bFullyLoadUsedTextures = CVarStreamingFullyLoadUsedTextures.GetValueOnAnyThread() != 0;
+	bUseAllMips = CVarStreamingUseAllMips.GetValueOnAnyThread() != 0;
+
+	if (bUseAllMips)
+	{
+		bUsePerTextureBias = false;
+		MipBias = 0;
+	}
+}
+
 
 extern ENGINE_API UPrimitiveComponent* GDebugSelectedComponent;
 
diff --git a/Engine/Source/Runtime/Engine/Private/Streaming/TextureStreamingHelpers.h b/Engine/Source/Runtime/Engine/Private/Streaming/TextureStreamingHelpers.h
--- a/Engine/Source/Runtime/Engine/Private/Streaming/TextureStreamingHelpers.h
+++ b/Engine/Source/Runtime/Engine/Private/Streaming/TextureStreamingHelpers.h
@@ -48,15 +48,38 @@ extern TAutoConsoleVariable<int32> CVarSetTextureStreaming;
 #endif
 
 extern TAutoConsoleVariable<float> CVarStreamingBoost;
-extern TAutoConsoleVariable<float> CVarStreamingScreenSizeEffectiveMax;
 extern TAutoConsoleVariable<int32> CVarStreamingUseFixedPoolSize;
 extern TAutoConsoleVariable<int32> CVarStreamingPoolSize;
-extern TAutoConsoleVariable<int32> CVarStreamingMaxTempMemoryAllowed;
-extern TAutoConsoleVariable<int32> CVarStreamingDropMips;
-extern TAutoConsoleVariable<int32> CVarStreamingHLODStrategy;
-extern TAutoConsoleVariable<float> CVarStreamingHiddenPrimitiveScale;
-extern TAutoConsoleVariable<float> CVarStreamingMipBias;
-extern TAutoConsoleVariable<int32> CVarOnlyStreamInTextures;
+
+struct FTextureStreamingSettings
+{
+	FTextureStreamingSettings() { Update(); }
+	void Update();
+
+	FORCEINLINE bool operator ==(const FTextureStreamingSettings& Rhs) const { return FMemory::Memcmp(this, &Rhs, sizeof(FTextureStreamingSettings)) == 0; }
+	FORCEINLINE bool operator !=(const FTextureStreamingSettings& Rhs) const { return FMemory::Memcmp(this, &Rhs, sizeof(FTextureStreamingSettings)) != 0; }
+
+	FORCEINLINE float GlobalMipBias() const { return !bUsePerTextureBias ? MipBias : 0; }
+	FORCEINLINE int32 GlobalMipBiasAsInt() const { return !bUsePerTextureBias ? FMath::FloorToInt(MipBias) : 0; }
+	FORCEINLINE float GlobalMipBiasAsScale() const { return !bUsePerTextureBias ? FMath::Exp2(-MipBias) : 1.f; }
+
+	FORCEINLINE int32 MaxPerTextureMipBias() const { return bUsePerTextureBias ? FMath::FloorToInt(MipBias) : 0; }
+
+	float MaxEffectiveScreenSize;
+	int32 MaxTempMemoryAllowed;
+	int32 DropMips;
+	int32 HLODStrategy;
+	float HiddenPrimitiveScale;
+	int32 PoolSize;
+	bool bUseNewMetrics;
+	bool bFullyLoadUsedTextures;
+	bool bUseAllMips;
+
+protected:
+
+	bool bUsePerTextureBias;
+	float MipBias;
+};
 
 typedef TArray<int32, TMemStackAllocator<> > FStreamingRequests;
 typedef TArray<const UTexture2D*, TInlineAllocator<12> > FRemovedTextureArray;
diff --git a/Engine/Source/Runtime/Engine/Private/Texture.cpp b/Engine/Source/Runtime/Engine/Private/Texture.cpp
--- a/Engine/Source/Runtime/Engine/Private/Texture.cpp
+++ b/Engine/Source/Runtime/Engine/Private/Texture.cpp
@@ -58,7 +58,7 @@ UTexture::UTexture(const FObjectInitializer& ObjectInitializer)
 	CompositeTextureMode = CTM_NormalRoughnessToAlpha;
 	CompositePower = 1.0f;
 	bUseLegacyGamma = false;
-
+	AlphaCoverageThresholds = FVector4(0, 0, 0, 0);
 	PaddingColor = FColor::Black;
 	ChromaKeyColor = FColorList::Magenta;
 	ChromaKeyThreshold = 1.0f / 255.0f;
@@ -470,10 +470,6 @@ bool UTexture::ForceUpdateTextureStreaming()
 {
 	if (!IStreamingManager::HasShutdown())
 	{
-		// Make sure textures can be streamed out so that we can unload current mips.
-		const bool bOldOnlyStreamInTextures = CVarOnlyStreamInTextures->GetInt() != 0;
-		CVarOnlyStreamInTextures->Set(false, ECVF_SetByCode);
-
 #if WITH_EDITOR
 		for( TObjectIterator<UTexture2D> It; It; ++It )
 		{
@@ -490,9 +486,6 @@ bool UTexture::ForceUpdateTextureStreaming()
 		IStreamingManager::Get().UpdateResourceStreaming( 0 );
 		// Block till requests are finished.
 		IStreamingManager::Get().BlockTillAllRequestsFinished();
-
-		// Restore streaming out of textures.
-		CVarOnlyStreamInTextures->Set(bOldOnlyStreamInTextures, ECVF_SetByCode);
 	}
 
 	return true;
diff --git a/Engine/Source/Runtime/Engine/Private/Texture2D.cpp b/Engine/Source/Runtime/Engine/Private/Texture2D.cpp
--- a/Engine/Source/Runtime/Engine/Private/Texture2D.cpp
+++ b/Engine/Source/Runtime/Engine/Private/Texture2D.cpp
@@ -329,6 +329,11 @@ void UTexture2D::LinkStreaming()
 	{
 		IStreamingManager::Get().GetTextureStreamingManager().AddStreamingTexture(this);
 	}
+	else
+	{
+		StreamingIndex = INDEX_NONE;
+	}
+
 }
 
 void UTexture2D::UnlinkStreaming()
@@ -911,8 +916,7 @@ bool UTexture2D::ShouldMipLevelsBeForcedResident() const
 	{
 		return true;
 	}
-	float CurrentTime = float(FPlatformTime::Seconds() - GStartTime);
-	if ( ForceMipLevelsToBeResidentTimestamp >= CurrentTime )
+	if ( ForceMipLevelsToBeResidentTimestamp >= FApp::GetCurrentTime() )
 	{
 		return true;
 	}
@@ -976,7 +980,7 @@ void UTexture2D::SetForceMipLevelsToBeResident( float Seconds, int32 CinematicTe
 	uint32 TextureGroupBitfield = (uint32) CinematicTextureGroups;
 	uint32 MyTextureGroup = FMath::BitFlag[LODGroup];
 	bUseCinematicMipLevels = (TextureGroupBitfield & MyTextureGroup) ? true : false;
-	ForceMipLevelsToBeResidentTimestamp = float(FPlatformTime::Seconds() - GStartTime) + Seconds;
+	ForceMipLevelsToBeResidentTimestamp = FApp::GetCurrentTime() + Seconds;
 }
 
 int32 UTexture2D::Blueprint_GetSizeX() const
diff --git a/Engine/Source/Runtime/Engine/Private/TextureDerivedData.cpp b/Engine/Source/Runtime/Engine/Private/TextureDerivedData.cpp
--- a/Engine/Source/Runtime/Engine/Private/TextureDerivedData.cpp
+++ b/Engine/Source/Runtime/Engine/Private/TextureDerivedData.cpp
@@ -65,6 +65,7 @@ static void SerializeForKey(FArchive& Ar, const FTextureBuildSettings& Settings)
 	float TempFloat;
 	uint8 TempByte;
 	FColor TempColor;
+	FVector4 TempVector4;
 
 	TempFloat = Settings.ColorAdjustment.AdjustBrightness; Ar << TempFloat;
 	TempFloat = Settings.ColorAdjustment.AdjustBrightnessCurve; Ar << TempFloat;
@@ -83,6 +84,12 @@ static void SerializeForKey(FArchive& Ar, const FTextureBuildSettings& Settings)
 	TempByte = Settings.bSRGB ? (Settings.bSRGB | ( Settings.bUseLegacyGamma ? 0 : 0x2 )) : 0; Ar << TempByte;
 	TempByte = Settings.bPreserveBorder; Ar << TempByte;
 	TempByte = Settings.bDitherMipMapAlpha; Ar << TempByte;
+
+	if (Settings.AlphaCoverageThresholds != FVector4(0, 0, 0, 0))
+	{
+		TempVector4 = Settings.AlphaCoverageThresholds; Ar << TempVector4;
+	}
+	
 	TempByte = Settings.bComputeBokehAlpha; Ar << TempByte;
 	TempByte = Settings.bReplicateRed; Ar << TempByte;
 	TempByte = Settings.bReplicateAlpha; Ar << TempByte;
@@ -257,6 +264,7 @@ static void GetTextureBuildSettings(
 	OutBuildSettings.bUseLegacyGamma = Texture.bUseLegacyGamma;
 	OutBuildSettings.bPreserveBorder = Texture.bPreserveBorder;
 	OutBuildSettings.bDitherMipMapAlpha = Texture.bDitherMipMapAlpha;
+	OutBuildSettings.AlphaCoverageThresholds = Texture.AlphaCoverageThresholds;
 	OutBuildSettings.bComputeBokehAlpha = (Texture.LODGroup == TEXTUREGROUP_Bokeh);
 	OutBuildSettings.bReplicateAlpha = false;
 	OutBuildSettings.bReplicateRed = false;
diff --git a/Engine/Source/Runtime/Engine/Private/UnrealEngine.cpp b/Engine/Source/Runtime/Engine/Private/UnrealEngine.cpp
--- a/Engine/Source/Runtime/Engine/Private/UnrealEngine.cpp
+++ b/Engine/Source/Runtime/Engine/Private/UnrealEngine.cpp
@@ -992,7 +992,8 @@ void UEngine::Init(IEngineLoop* InEngineLoop)
 
 	// Dynamically load engine runtime modules
 	{
-		FModuleManager::Get().LoadModuleChecked( TEXT( "StreamingPauseRendering" ) );
+		FModuleManager::Get().LoadModuleChecked(TEXT("StreamingPauseRendering"));
+		FModuleManager::Get().LoadModuleChecked(TEXT("Niagara"));
 	}
 
 	bool bIsRHS = true;
diff --git a/Engine/Source/Runtime/Engine/Public/GPUSkinVertexFactory.h b/Engine/Source/Runtime/Engine/Public/GPUSkinVertexFactory.h
--- a/Engine/Source/Runtime/Engine/Public/GPUSkinVertexFactory.h
+++ b/Engine/Source/Runtime/Engine/Public/GPUSkinVertexFactory.h
@@ -290,12 +290,8 @@ public:
 			, CurrentFrameNumber(0)
 		{
 			// BoneDataOffset and BoneTextureSize are not set as they are only valid if IsValidRef(BoneTexture)
-			if(!MaxBonesVar)
-			{
-				MaxBonesVar = IConsoleManager::Get().FindTConsoleVariableDataInt(TEXT("Compat.MAX_GPUSKIN_BONES"));
-				MaxGPUSkinBones = MaxBonesVar->GetValueOnGameThread();
-				check(MaxGPUSkinBones <= 256);
-			}
+			MaxGPUSkinBones = GetMaxGPUSkinBones();
+			check(MaxGPUSkinBones <= GHardwareMaxGPUSkinBones);
 		}
 
 		/** Mesh origin and Mesh Extension for Mesh compressions **/
@@ -400,8 +396,9 @@ public:
 	};
 
 	FGPUBaseSkinVertexFactory(ERHIFeatureLevel::Type InFeatureLevel)
-	:	FVertexFactory(InFeatureLevel)
-	{}
+		: FVertexFactory(InFeatureLevel)
+	{
+	}
 
 	/** accessor */
 	FORCEINLINE FShaderDataType& GetShaderData()
@@ -423,7 +420,14 @@ public:
 		return (FSkeletalMeshVertexBuffer*)(Streams[0].VertexBuffer);
 	}
 
+	ENGINE_API static int32 GetMaxGPUSkinBones();
+
+	static const uint32 GHardwareMaxGPUSkinBones = 256;	
+	
 protected:
+
+	
+
 	/** dynamic data need for setting the shader */ 
 	FShaderDataType ShaderData;
 	/** Pool of buffers for bone matrices. */
diff --git a/Engine/Source/Runtime/Engine/Public/MaterialCompiler.h b/Engine/Source/Runtime/Engine/Public/MaterialCompiler.h
--- a/Engine/Source/Runtime/Engine/Public/MaterialCompiler.h
+++ b/Engine/Source/Runtime/Engine/Public/MaterialCompiler.h
@@ -197,7 +197,7 @@ public:
 	virtual int32 PerInstanceRandom() = 0;
 	virtual int32 PerInstanceFadeAmount() = 0;
 	virtual int32 AntialiasedTextureMask(int32 Tex, int32 UV, float Threshold, uint8 Channel) = 0;
-	virtual int32 Noise(int32 Position, float Scale, int32 Quality, uint8 NoiseFunction, bool bTurbulence, int32 Levels, float OutputMin, float OutputMax, float LevelScale, int32 FilterWidth) = 0;
+	virtual int32 Noise(int32 Position, float Scale, int32 Quality, uint8 NoiseFunction, bool bTurbulence, int32 Levels, float OutputMin, float OutputMax, float LevelScale, int32 FilterWidth, bool bTiling, uint32 RepeatSize) = 0;
 	virtual int32 BlackBody( int32 Temp ) = 0;
 	virtual int32 DistanceToNearestSurface(int32 PositionArg) = 0;
 	virtual int32 DistanceFieldGradient(int32 PositionArg) = 0;
@@ -371,9 +371,9 @@ public:
 	{
 		return Compiler->AntialiasedTextureMask(Tex, UV, Threshold, Channel);
 	}
-	virtual int32 Noise(int32 Position, float Scale, int32 Quality, uint8 NoiseFunction, bool bTurbulence, int32 Levels, float OutputMin, float OutputMax, float LevelScale, int32 FilterWidth) override
+	virtual int32 Noise(int32 Position, float Scale, int32 Quality, uint8 NoiseFunction, bool bTurbulence, int32 Levels, float OutputMin, float OutputMax, float LevelScale, int32 FilterWidth, bool bTiling, uint32 RepeatSize) override
 	{
-		return Compiler->Noise(Position, Scale, Quality, NoiseFunction, bTurbulence, Levels, OutputMin, OutputMax, LevelScale, FilterWidth);
+		return Compiler->Noise(Position, Scale, Quality, NoiseFunction, bTurbulence, Levels, OutputMin, OutputMax, LevelScale, FilterWidth, bTiling, RepeatSize);
 	}
 	virtual int32 BlackBody( int32 Temp ) override { return Compiler->BlackBody(Temp); }
 	virtual int32 DistanceToNearestSurface(int32 PositionArg) override { return Compiler->DistanceToNearestSurface(PositionArg); }
diff --git a/Engine/Source/Runtime/Engine/Public/MeshParticleVertexFactory.h b/Engine/Source/Runtime/Engine/Public/MeshParticleVertexFactory.h
--- a/Engine/Source/Runtime/Engine/Public/MeshParticleVertexFactory.h
+++ b/Engine/Source/Runtime/Engine/Public/MeshParticleVertexFactory.h
@@ -13,7 +13,7 @@
 /**
  * Uniform buffer for mesh particle vertex factories.
  */
-BEGIN_UNIFORM_BUFFER_STRUCT( FMeshParticleUniformParameters, )
+BEGIN_UNIFORM_BUFFER_STRUCT( FMeshParticleUniformParameters, ENGINE_API)
 	DECLARE_UNIFORM_BUFFER_STRUCT_MEMBER( FVector4, SubImageSize )
 	DECLARE_UNIFORM_BUFFER_STRUCT_MEMBER( uint32, TexCoordWeightA )
 	DECLARE_UNIFORM_BUFFER_STRUCT_MEMBER( uint32, TexCoordWeightB )
@@ -24,7 +24,7 @@ typedef TUniformBufferRef<FMeshParticleUniformParameters> FMeshParticleUniformBu
 /**
  * Vertex factory for rendering instanced mesh particles with out dynamic parameter support.
  */
-class FMeshParticleVertexFactory : public FParticleVertexFactoryBase
+class ENGINE_API FMeshParticleVertexFactory : public FParticleVertexFactoryBase
 {
 	DECLARE_VERTEX_FACTORY_TYPE(FMeshParticleVertexFactory);
 public:
@@ -110,7 +110,7 @@ public:
 	/**
 	 * An implementation of the interface used by TSynchronizedResource to update the resource with new data from the game thread.
 	 */
-	ENGINE_API void SetData(const FDataType& InData);
+	void SetData(const FDataType& InData);
 
 	/**
 	 * Set the uniform buffer for this vertex factory.
@@ -178,7 +178,7 @@ protected:
 };
 
 
-class FMeshParticleVertexFactoryEmulatedInstancing : public FMeshParticleVertexFactory
+class ENGINE_API FMeshParticleVertexFactoryEmulatedInstancing : public FMeshParticleVertexFactory
 {
 	DECLARE_VERTEX_FACTORY_TYPE(FMeshParticleVertexFactoryEmulatedInstancing);
 
diff --git a/Engine/Source/Runtime/Engine/Public/ParticleBeamTrailVertexFactory.h b/Engine/Source/Runtime/Engine/Public/ParticleBeamTrailVertexFactory.h
--- a/Engine/Source/Runtime/Engine/Public/ParticleBeamTrailVertexFactory.h
+++ b/Engine/Source/Runtime/Engine/Public/ParticleBeamTrailVertexFactory.h
@@ -12,7 +12,7 @@
 /**
  * Uniform buffer for particle beam/trail vertex factories.
  */
-BEGIN_UNIFORM_BUFFER_STRUCT( FParticleBeamTrailUniformParameters, )
+BEGIN_UNIFORM_BUFFER_STRUCT( FParticleBeamTrailUniformParameters, ENGINE_API)
 	DECLARE_UNIFORM_BUFFER_STRUCT_MEMBER( FVector4, CameraRight )
 	DECLARE_UNIFORM_BUFFER_STRUCT_MEMBER( FVector4, CameraUp )
 	DECLARE_UNIFORM_BUFFER_STRUCT_MEMBER( FVector4, ScreenAlignment )
diff --git a/Engine/Source/Runtime/Engine/Public/ParticleVertexFactory.h b/Engine/Source/Runtime/Engine/Public/ParticleVertexFactory.h
--- a/Engine/Source/Runtime/Engine/Public/ParticleVertexFactory.h
+++ b/Engine/Source/Runtime/Engine/Public/ParticleVertexFactory.h
@@ -78,7 +78,7 @@ private:
 /**
  * Uniform buffer for particle sprite vertex factories.
  */
-BEGIN_UNIFORM_BUFFER_STRUCT( FParticleSpriteUniformParameters, )
+BEGIN_UNIFORM_BUFFER_STRUCT( FParticleSpriteUniformParameters, ENGINE_API)
 	DECLARE_UNIFORM_BUFFER_STRUCT_MEMBER_EX( FVector4, AxisLockRight, EShaderPrecisionModifier::Half )
 	DECLARE_UNIFORM_BUFFER_STRUCT_MEMBER_EX( FVector4, AxisLockUp, EShaderPrecisionModifier::Half )
 	DECLARE_UNIFORM_BUFFER_STRUCT_MEMBER_EX( FVector4, TangentSelector, EShaderPrecisionModifier::Half )
diff --git a/Engine/Source/Runtime/Engine/Public/PrimitiveViewRelevance.h b/Engine/Source/Runtime/Engine/Public/PrimitiveViewRelevance.h
--- a/Engine/Source/Runtime/Engine/Public/PrimitiveViewRelevance.h
+++ b/Engine/Source/Runtime/Engine/Public/PrimitiveViewRelevance.h
@@ -9,6 +9,10 @@
  */
 struct FPrimitiveViewRelevance
 {
+	// Warning: This class is memzeroed externally as 0 is assumed a
+	// valid value for all members meaning 'not relevant'. If this
+	// changes existing class usage should be re-evaluated
+
 	// from FMaterialRelevance (could be made the base class):
 
 	/** The LightingProfile supported by this primitive, as a bitmask. */
@@ -48,12 +52,14 @@ struct FPrimitiveViewRelevance
 	uint32 bEditorNoDepthTestPrimitiveRelevance : 1;
 	/** The primitive should have GatherSimpleLights called on the proxy when gathering simple lights. */
 	uint32 bHasSimpleLights : 1;
-	/**  The primitive has one or more elements that have World Position Offset. */
+	/** The primitive has one or more elements that have World Position Offset. */
 	uint32 bUsesWorldPositionOffset : 1;
-	/** */
+	/** Whether the primitive uses non-default lighting channels. */
 	uint32 bUsesLightingChannels : 1;
 	/** */
 	uint32 bDecal : 1;
+	/** Whether the primitive has materals that use translucent surface lighting. */
+	uint32 bTranslucentSurfaceLighting : 1;
 
 	/** 
 	 * Whether this primitive view relevance has been initialized this frame.  
diff --git a/Engine/Source/Runtime/Engine/Public/SceneInterface.h b/Engine/Source/Runtime/Engine/Public/SceneInterface.h
--- a/Engine/Source/Runtime/Engine/Public/SceneInterface.h
+++ b/Engine/Source/Runtime/Engine/Public/SceneInterface.h
@@ -18,6 +18,14 @@ enum EBasePassDrawListType
 	EBasePass_Masked,
 	EBasePass_MAX
 };
+
+enum class EShadingPath
+{
+	Mobile,
+	Deferred,
+	Num,
+};
+
 /**
  * An interface to the private scene manager implementation of a scene.  Use GetRendererModule().AllocateScene to create.
  * The scene
@@ -335,14 +343,21 @@ public:
 	virtual ERHIFeatureLevel::Type GetFeatureLevel() const { return GMaxRHIFeatureLevel; }
 	EShaderPlatform GetShaderPlatform() const { return GShaderPlatformForFeatureLevel[GetFeatureLevel()]; }
 
-	static bool ShouldUseDeferredRenderer(ERHIFeatureLevel::Type InFeatureLevel)
+	static EShadingPath GetShadingPath(ERHIFeatureLevel::Type InFeatureLevel)
 	{
-		return InFeatureLevel >= ERHIFeatureLevel::SM4;
+		if (InFeatureLevel >= ERHIFeatureLevel::SM4)
+		{
+			return EShadingPath::Deferred;
+		}
+		else
+		{
+			return EShadingPath::Mobile;
+		}
 	}
 
-	bool ShouldUseDeferredRenderer() const
+	EShadingPath GetShadingPath() const
 	{
-		return ShouldUseDeferredRenderer(GetFeatureLevel());
+		return GetShadingPath(GetFeatureLevel());
 	}
 
 #if WITH_EDITOR
diff --git a/Engine/Source/Runtime/Niagara/Classes/NiagaraEffect.h b/Engine/Source/Runtime/Niagara/Classes/NiagaraEffect.h
--- a/Engine/Source/Runtime/Niagara/Classes/NiagaraEffect.h
+++ b/Engine/Source/Runtime/Niagara/Classes/NiagaraEffect.h
@@ -3,9 +3,8 @@
 #pragma once
 #include "NiagaraSimulation.h"
 #include "NiagaraEmitterProperties.h"
-#include "Engine/NiagaraEffectRenderer.h"
+#include "NiagaraEffectRenderer.h"
 #include "List.h"
-#include "Runtime/VectorVM/Public/VectorVMDataObject.h"
 #include "NiagaraEffect.generated.h"
 
 
@@ -70,7 +69,6 @@ public:
 			Emitters.Add(MakeShareable(Sim));
 		}
 		InitRenderModules(Component->GetWorld()->FeatureLevel);
-		VolumeGrid = NewObject<UNiagaraSparseVolumeDataObject>();
 	}
 
 	explicit FNiagaraEffectInstance(UNiagaraEffect *InAsset)
@@ -78,7 +76,6 @@ public:
 		, Effect(InAsset)
 	{
 		InitEmitters(InAsset);
-		VolumeGrid = NewObject<UNiagaraSparseVolumeDataObject>();
 	}
 
 	NIAGARA_API void InitEmitters(UNiagaraEffect *InAsset);
@@ -137,12 +134,7 @@ public:
 	{
 		Constants.SetOrAdd(ID, Value);
 	}
-
-	void SetConstant(FName ConstantName, UNiagaraDataObject *Value)
-	{
-		Constants.SetOrAdd(ConstantName, Value);
-	}
-
+	
 	void Tick(float DeltaSeconds);
 
 	NIAGARA_API void RenderModuleupdate();
@@ -175,6 +167,4 @@ private:
 	TMap<FNiagaraDataSetID, FNiagaraDataSet> ExternalEvents;
 
 	TArray< TSharedPtr<FNiagaraSimulation> > Emitters;
-
-	UNiagaraSparseVolumeDataObject *VolumeGrid;
 };
\ No newline at end of file
diff --git a/Engine/Source/Runtime/Niagara/NiagaraConstantSet.h b/Engine/Source/Runtime/Niagara/NiagaraConstantSet.h
--- a/Engine/Source/Runtime/Niagara/NiagaraConstantSet.h
+++ b/Engine/Source/Runtime/Niagara/NiagaraConstantSet.h
@@ -3,7 +3,6 @@
 #pragma once
 
 #include "NiagaraCommon.h"
-#include "VectorVMDataObject.h"
 #include "NiagaraConstantSet.generated.h"
 
 struct FNiagaraConstants;
@@ -17,7 +16,6 @@ private:
 	TMap<FNiagaraVariableInfo, float> ScalarConstants;
 	TMap<FNiagaraVariableInfo, FVector4> VectorConstants;
 	TMap<FNiagaraVariableInfo, FMatrix> MatrixConstants;
-	TMap<FNiagaraVariableInfo, class UNiagaraDataObject*> DataConstants;
 
 public:
 	FNiagaraConstantMap()
@@ -29,28 +27,22 @@ public:
 	void SetOrAdd(FNiagaraVariableInfo ID, float Sc);
 	void SetOrAdd(FNiagaraVariableInfo ID, const FVector4& Vc);
 	void SetOrAdd(FNiagaraVariableInfo ID, const FMatrix& Mc);
-	void SetOrAdd(FNiagaraVariableInfo ID, UNiagaraDataObject* Cc);
 	void SetOrAdd(FName Name, float Sc);
 	void SetOrAdd(FName Name, const FVector4& Vc);
 	void SetOrAdd(FName Name, const FMatrix& Mc);
-	void SetOrAdd(FName Name, class UNiagaraDataObject* Cc);
 
 	float* FindScalar(FNiagaraVariableInfo ID);
 	FVector4* FindVector(FNiagaraVariableInfo ID);
 	FMatrix* FindMatrix(FNiagaraVariableInfo ID);
-	UNiagaraDataObject* FindDataObj(FNiagaraVariableInfo ID);
 	const float* FindScalar(FNiagaraVariableInfo ID)const;
 	const FVector4* FindVector(FNiagaraVariableInfo ID)const;
 	const FMatrix* FindMatrix(FNiagaraVariableInfo ID)const;
-	UNiagaraDataObject* FindDataObj(FNiagaraVariableInfo ID) const;
 	float* FindScalar(FName Name);
 	FVector4* FindVector(FName Name);
 	FMatrix* FindMatrix(FName Name);
-	UNiagaraDataObject* FindDataObj(FName Name);
 	const float* FindScalar(FName Name)const;
 	const FVector4* FindVector(FName Name)const;
 	const FMatrix* FindMatrix(FName Name)const;
-	UNiagaraDataObject* FindDataObj(FName Name)const;
 
 	void Merge(FNiagaraConstants &InConstants);
 	void Merge(FNiagaraConstantMap &InMap);
@@ -58,17 +50,10 @@ public:
 	const TMap<FNiagaraVariableInfo, float> &GetScalarConstants()	const;
 	const TMap<FNiagaraVariableInfo, FVector4> &GetVectorConstants()	const;
 	const TMap<FNiagaraVariableInfo, FMatrix> &GetMatrixConstants()	const;
-	const TMap<FNiagaraVariableInfo, class UNiagaraDataObject*> &GetDataConstants()	const;
 
 
 	virtual bool Serialize(FArchive &Ar)
 	{
-		// TODO: can't serialize the data object constants at the moment; need to figure that out
-		Ar << ScalarConstants << VectorConstants << MatrixConstants;
-		if (Ar.UE4Ver() >= VER_UE4_NIAGARA_DATA_OBJECT_DEV_UI_FIX)
-		{
-			Ar << DataConstants;
-		}
 		return true;
 	}
 };
@@ -144,24 +129,6 @@ struct FNiagaraConstants_Matrix : public FNiagaraConstantBase
 	FMatrix Value;
 };
 
-USTRUCT()
-struct FNiagaraConstants_DataObject : public FNiagaraConstantBase
-{
-	GENERATED_USTRUCT_BODY()
-
-	FNiagaraConstants_DataObject()
-	: Value(nullptr)
-	{}
-
-	FNiagaraConstants_DataObject(FName InName, UNiagaraDataObject* InValue)
-	: FNiagaraConstantBase(InName)
-	, Value(InValue)
-	{}
-
-	UPROPERTY(Instanced, VisibleAnywhere, Category = "Constant")
-	UNiagaraDataObject* Value;
-};
-
 //Dummy struct used to serialize in the old layout of FNiagaraConstants
 //This should be removed once everyone has recompiled and saved their scripts.
 USTRUCT()
@@ -176,16 +143,12 @@ struct FDeprecatedNiagaraConstants
 	TArray<FNiagaraVariableInfo> VectorConstantsInfo_DEPRECATED;
 	UPROPERTY()
 	TArray<FNiagaraVariableInfo> MatrixConstantsInfo_DEPRECATED;
-	//UPROPERTY()
-	TArray<FNiagaraVariableInfo> DataObjectConstantsInfo_DEPRECATED;
 	UPROPERTY()
 	TArray<float> ScalarConstants_DEPRECATED;
 	UPROPERTY()
 	TArray<FVector4> VectorConstants_DEPRECATED;
 	UPROPERTY()
 	TArray<FMatrix> MatrixConstants_DEPRECATED;
-	//UPROPERTY()
-	TArray<void*> DataObjectConstants_DEPRECATED;
 };
 
 USTRUCT()
@@ -201,19 +164,15 @@ public:
 	TArray<FNiagaraConstants_Vector> VectorConstants;
 	UPROPERTY(EditAnywhere, EditFixedSize, Category = "Constant")
 	TArray<FNiagaraConstants_Matrix> MatrixConstants;
-	UPROPERTY(EditAnywhere, EditFixedSize, Category = "Constant")
-	TArray<FNiagaraConstants_DataObject> DataObjectConstants;
 
 	NIAGARA_API void Empty();
 	NIAGARA_API void GetScalarConstant(int32 Index, float& OutValue, FNiagaraVariableInfo& OutInfo);
 	NIAGARA_API void GetVectorConstant(int32 Index, FVector4& OutValue, FNiagaraVariableInfo& OutInfo);
 	NIAGARA_API void GetMatrixConstant(int32 Index, FMatrix& OutValue, FNiagaraVariableInfo& OutInfo);
-	NIAGARA_API void GetDataObjectConstant(int32 Index, UNiagaraDataObject*& OutValue, FNiagaraVariableInfo& OutInfo);
 
 	const int32 GetNumScalarConstants()const{ return ScalarConstants.Num(); }
 	const int32 GetNumVectorConstants()const{ return VectorConstants.Num(); }
 	const int32 GetNumMatrixConstants()const{ return MatrixConstants.Num(); }
-	const int32 GetNumDataObjectConstants()const{ return DataObjectConstants.Num(); }
 
 	NIAGARA_API void Init(class UNiagaraEmitterProperties* EmitterProps, struct FNiagaraEmitterScriptProperties* ScriptProps);
 
@@ -226,23 +185,18 @@ public:
 	/** Fills the entire constants set into the constant table. */
 	void AppendToConstantsTable(TArray<FVector4>& ConstantsTable)const;
 
-	void AppendBufferConstants(TArray<class UNiagaraDataObject*> &Objs) const;
-
 	/** 
 	Fills only selected constants into the table. In the order they appear in the array of passed names not the order they appear in the set. 
 	Checks the passed map for entires to superceed the default values in the set.
 	*/
 	void AppendToConstantsTable(TArray<FVector4>& ConstantsTable, const FNiagaraConstantMap& Externals)const;
-	void AppendExternalBufferConstants(TArray<class UNiagaraDataObject*> &DataObjs, const FNiagaraConstantMap& Externals) const;
 
 	NIAGARA_API void SetOrAdd(const FNiagaraVariableInfo& Constant, float Value);
 	NIAGARA_API void SetOrAdd(const FNiagaraVariableInfo& Constant, const FVector4& Value);
 	NIAGARA_API void SetOrAdd(const FNiagaraVariableInfo& Constant, const FMatrix& Value);
-	NIAGARA_API void SetOrAdd(const FNiagaraVariableInfo& Constant, UNiagaraDataObject* Value);
 	NIAGARA_API void SetOrAdd(FName Name, float Value);
 	NIAGARA_API void SetOrAdd(FName Name, const FVector4& Value);
 	NIAGARA_API void SetOrAdd(FName Name, const FMatrix& Value);
-	NIAGARA_API void SetOrAdd(FName Name, UNiagaraDataObject* Value);
 
 	FORCEINLINE int32 NumScalars()const { return ScalarConstants.Num(); }
 	FORCEINLINE int32 NumVectors()const { return VectorConstants.Num(); }
@@ -257,7 +211,6 @@ public:
 	NIAGARA_API float* FindScalar(FName Name);
 	NIAGARA_API FVector4* FindVector(FName Name);
 	NIAGARA_API FMatrix* FindMatrix(FName Name);
-	NIAGARA_API UNiagaraDataObject* FindDataObj(FName Name);
 
 	/** Return the first constant used for this constant in a table. */
 	FORCEINLINE int32 GetTableIndex_Scalar(const FNiagaraVariableInfo& Constant)const
@@ -289,9 +242,4 @@ public:
 		}
 		return INDEX_NONE;
 	}
-
-	FORCEINLINE int32 GetTableIndex_DataObj(const FNiagaraVariableInfo& Constant)const
-	{ 
-		return DataObjectConstants.IndexOfByPredicate([&](const FNiagaraConstants_DataObject& V){ return Constant.Name == V.Name; });
-	}
 };
diff --git a/Engine/Source/Runtime/Niagara/NiagaraScriptConstantData.h b/Engine/Source/Runtime/Niagara/NiagaraScriptConstantData.h
--- a/Engine/Source/Runtime/Niagara/NiagaraScriptConstantData.h
+++ b/Engine/Source/Runtime/Niagara/NiagaraScriptConstantData.h
@@ -34,16 +34,14 @@ struct FNiagaraScriptConstantData
 	FNiagaraConstants& GetExternalConstants(){ return ExternalConstantsNew; }
 
 	/** Fill a constants table ready for use in an update or spawn. */
-	void FillConstantTable(const FNiagaraConstantMap& ExternalConstantsMap, TArray<FVector4>& OutConstantTable, TArray<UNiagaraDataObject *> &DataObjTable) const
+	void FillConstantTable(const FNiagaraConstantMap& ExternalConstantsMap, TArray<FVector4>& OutConstantTable) const
 	{
 		//First up in the table comes the External constants in scalar->vector->matrix order.
 		//Only fills the constants actually used by the script.
 		OutConstantTable.Empty();
 		ExternalConstantsNew.AppendToConstantsTable(OutConstantTable, ExternalConstantsMap);
-		ExternalConstantsNew.AppendExternalBufferConstants(DataObjTable, ExternalConstantsMap);
 		//Next up add all the internal constants from the script.
 		InternalConstantsNew.AppendToConstantsTable(OutConstantTable);
-		InternalConstantsNew.AppendBufferConstants(DataObjTable);
 	}
 
 	template<typename T>
@@ -79,13 +77,6 @@ struct FNiagaraScriptConstantData
 			{
 				ConstIdx = Consts.GetTableIndex_Matrix(Constant);
 				Type = ENiagaraDataType::Matrix;
-
-				if (ConstIdx == INDEX_NONE)	// curves/buffers are in a separate table, so set base to 0
-				{
-					ConstIdx = Consts.GetTableIndex_DataObj(Constant);
-					Type = ENiagaraDataType::Curve;
-					Base = 0;
-				}
 			}
 		}
 
diff --git a/Engine/Source/Runtime/Niagara/Private/NiagaraConstantSet.cpp b/Engine/Source/Runtime/Niagara/Private/NiagaraConstantSet.cpp
--- a/Engine/Source/Runtime/Niagara/Private/NiagaraConstantSet.cpp
+++ b/Engine/Source/Runtime/Niagara/Private/NiagaraConstantSet.cpp
@@ -42,12 +42,6 @@ void FNiagaraConstantMap::SetOrAdd(FNiagaraVariableInfo ID, const FMatrix& Mc)
 	MatrixConstants.FindOrAdd(ID) = Mc;
 }
 
-void FNiagaraConstantMap::SetOrAdd(FNiagaraVariableInfo ID, UNiagaraDataObject* Cc)
-{
-	check(ID.Type == ENiagaraDataType::Curve);
-	DataConstants.FindOrAdd(ID) = Cc;
-}
-
 void FNiagaraConstantMap::SetOrAdd(FName Name, float Sc)
 {
 	ScalarConstants.FindOrAdd(FNiagaraVariableInfo(Name, ENiagaraDataType::Scalar)) = Sc;
@@ -63,46 +57,21 @@ void FNiagaraConstantMap::SetOrAdd(FName Name, const FMatrix& Mc)
 	MatrixConstants.FindOrAdd(FNiagaraVariableInfo(Name, ENiagaraDataType::Matrix)) = Mc;
 }
 
-void FNiagaraConstantMap::SetOrAdd(FName Name, class UNiagaraDataObject* Cc)
-{
-	DataConstants.FindOrAdd(FNiagaraVariableInfo(Name, ENiagaraDataType::Curve)) = Cc;
-}
-
 float* FNiagaraConstantMap::FindScalar(FNiagaraVariableInfo ID){ return ScalarConstants.Find(ID); }
 FVector4* FNiagaraConstantMap::FindVector(FNiagaraVariableInfo ID){ return VectorConstants.Find(ID); }
 FMatrix* FNiagaraConstantMap::FindMatrix(FNiagaraVariableInfo ID){ return MatrixConstants.Find(ID); }
-UNiagaraDataObject* FNiagaraConstantMap::FindDataObj(FNiagaraVariableInfo ID)
-{
-	UNiagaraDataObject * const *Ptr = DataConstants.Find(ID);
-	if (Ptr)
-	{
-		return *Ptr;
-	}
 
-	return nullptr;
-}
 const float* FNiagaraConstantMap::FindScalar(FNiagaraVariableInfo ID)const { return ScalarConstants.Find(ID); }
 const FVector4* FNiagaraConstantMap::FindVector(FNiagaraVariableInfo ID)const { return VectorConstants.Find(ID); }
 const FMatrix* FNiagaraConstantMap::FindMatrix(FNiagaraVariableInfo ID)const { return MatrixConstants.Find(ID); }
-UNiagaraDataObject* FNiagaraConstantMap::FindDataObj(FNiagaraVariableInfo ID) const
-{
-	UNiagaraDataObject * const *Ptr = DataConstants.Find(ID);
-	if (Ptr)
-	{
-		return *Ptr;
-	}
 
-	return nullptr;
-}
 float* FNiagaraConstantMap::FindScalar(FName Name) { return ScalarConstants.Find(FNiagaraVariableInfo(Name, ENiagaraDataType::Scalar)); }
 FVector4* FNiagaraConstantMap::FindVector(FName Name) { return VectorConstants.Find(FNiagaraVariableInfo(Name, ENiagaraDataType::Vector)); }
 FMatrix* FNiagaraConstantMap::FindMatrix(FName Name) { return MatrixConstants.Find(FNiagaraVariableInfo(Name, ENiagaraDataType::Matrix)); }
-UNiagaraDataObject* FNiagaraConstantMap::FindDataObj(FName Name){ return FindDataObj(FNiagaraVariableInfo(Name, ENiagaraDataType::Curve)); }
 
 const float* FNiagaraConstantMap::FindScalar(FName Name)const  { return ScalarConstants.Find(FNiagaraVariableInfo(Name, ENiagaraDataType::Scalar)); }
 const FVector4* FNiagaraConstantMap::FindVector(FName Name)const  { return VectorConstants.Find(FNiagaraVariableInfo(Name, ENiagaraDataType::Vector)); }
 const FMatrix* FNiagaraConstantMap::FindMatrix(FName Name)const  { return MatrixConstants.Find(FNiagaraVariableInfo(Name, ENiagaraDataType::Matrix)); }
-UNiagaraDataObject* FNiagaraConstantMap::FindDataObj(FName Name)const { return FindDataObj(FNiagaraVariableInfo(Name, ENiagaraDataType::Curve)); }
 
 void FNiagaraConstantMap::Merge(FNiagaraConstants &InConstants)
 {
@@ -118,10 +87,6 @@ void FNiagaraConstantMap::Merge(FNiagaraConstants &InConstants)
 	{
 		SetOrAdd(FNiagaraVariableInfo(C.Name, ENiagaraDataType::Matrix), C.Value);
 	}
-	for (FNiagaraConstants_DataObject& C : InConstants.DataObjectConstants)
-	{
-		SetOrAdd(FNiagaraVariableInfo(C.Name, ENiagaraDataType::Curve), C.Value);
-	}
 }
 
 void FNiagaraConstantMap::Merge(FNiagaraConstantMap &InMap)
@@ -140,11 +105,6 @@ void FNiagaraConstantMap::Merge(FNiagaraConstantMap &InMap)
 	{
 		SetOrAdd(MapIt.Key(), MapIt.Value());
 	}
-
-	for (auto MapIt = InMap.DataConstants.CreateIterator(); MapIt; ++MapIt)
-	{
-		SetOrAdd(MapIt.Key(), MapIt.Value());
-	}
 }
 
 const TMap<FNiagaraVariableInfo, float> &FNiagaraConstantMap::GetScalarConstants()	const
@@ -162,11 +122,6 @@ const TMap<FNiagaraVariableInfo, FMatrix> &FNiagaraConstantMap::GetMatrixConstan
 	return MatrixConstants;
 }
 
-const TMap<FNiagaraVariableInfo, class UNiagaraDataObject*> &FNiagaraConstantMap::GetDataConstants()	const
-{
-	return DataConstants;
-}
-
 //////////////////////////////////////////////////////////////////////////
 
 void FNiagaraConstants::Empty()
@@ -174,7 +129,6 @@ void FNiagaraConstants::Empty()
 	ScalarConstants.Empty();
 	VectorConstants.Empty();
 	MatrixConstants.Empty();
-	DataObjectConstants.Empty();
 }
 
 void FNiagaraConstants::GetScalarConstant(int32 Index, float& OutValue, FNiagaraVariableInfo& OutInfo)
@@ -195,12 +149,6 @@ void FNiagaraConstants::GetMatrixConstant(int32 Index, FMatrix& OutValue, FNiaga
 	OutValue = MatrixConstants[Index].Value;
 }
 
-void FNiagaraConstants::GetDataObjectConstant(int32 Index, UNiagaraDataObject*& OutValue, FNiagaraVariableInfo& OutInfo)
-{
-	OutInfo = FNiagaraVariableInfo(DataObjectConstants[Index].Name, ENiagaraDataType::Curve);
-	OutValue = DataObjectConstants[Index].Value;
-}
-
 void FNiagaraConstants::Init(UNiagaraEmitterProperties* EmitterProps, FNiagaraEmitterScriptProperties* ScriptProps)
 {
 	check(ScriptProps);
@@ -239,29 +187,6 @@ void FNiagaraConstants::Init(UNiagaraEmitterProperties* EmitterProps, FNiagaraEm
 
 	}
 	MatrixConstants = NewMatrixConstants;
-
-	const TArray<FNiagaraConstants_DataObject>& ScriptDataObjectConstants = ScriptConstants.DataObjectConstants;
-	TArray<FNiagaraConstants_DataObject> OldDataObjectConstants = DataObjectConstants;
-	DataObjectConstants.Empty();
-	for (const FNiagaraConstants_DataObject& ScriptConst : ScriptDataObjectConstants)
-	{
-		int32 OverrideIdx = OldDataObjectConstants.IndexOfByPredicate([&](const FNiagaraConstants_DataObject& C){ return ScriptConst.Name == C.Name && ScriptConst.Value && C.Value &&ScriptConst.Value->GetClass() == C.Value->GetClass(); });
-		
-		int32 AddIndex = DataObjectConstants.AddZeroed();
-		DataObjectConstants[AddIndex].Name = ScriptDataObjectConstants[AddIndex].Name;
-		
-		if (OverrideIdx != INDEX_NONE)
-		{
-			//If this constant already has data in the effect, use that.
-			DataObjectConstants[AddIndex].Value = OldDataObjectConstants[OverrideIdx].Value;
-		}		
-		else
-		{
-			//Otherwise, duplicate the data from the script.
-			if (ScriptConst.Value)
-				DataObjectConstants[AddIndex].Value = CastChecked<UNiagaraDataObject>(StaticDuplicateObject(ScriptConst.Value, EmitterProps));
-		}
-	}
 }
 
 /** 
@@ -273,7 +198,6 @@ void FNiagaraConstants::InitFromOldMap(FNiagaraConstantMap& OldMap)
 	const TMap<FNiagaraVariableInfo, float>& OldScalarConsts = OldMap.GetScalarConstants();
 	const TMap<FNiagaraVariableInfo, FVector4>& OldVectorConsts = OldMap.GetVectorConstants();
 	const TMap<FNiagaraVariableInfo, FMatrix>& OldMatrixConsts = OldMap.GetMatrixConstants();
-	const TMap<FNiagaraVariableInfo, UNiagaraDataObject*>& OldDataObjectConstants = OldMap.GetDataConstants();
 
 	for (const TPair<FNiagaraVariableInfo, float> Pair : OldScalarConsts)
 	{
@@ -287,10 +211,6 @@ void FNiagaraConstants::InitFromOldMap(FNiagaraConstantMap& OldMap)
 	{
 		MatrixConstants.Add(FNiagaraConstants_Matrix(Pair.Key.Name, Pair.Value));
 	}
-	for (const TPair<FNiagaraVariableInfo, UNiagaraDataObject*>& Pair : OldDataObjectConstants)
-	{
-		DataObjectConstants.Add(FNiagaraConstants_DataObject(Pair.Key.Name, Pair.Value));
-	}
 }
 
 void FNiagaraConstants::AppendToConstantsTable(TArray<FVector4>& ConstantsTable)const
@@ -318,14 +238,6 @@ void FNiagaraConstants::AppendToConstantsTable(TArray<FVector4>& ConstantsTable)
 	}
 }
 
-void FNiagaraConstants::AppendBufferConstants(TArray<class UNiagaraDataObject*> &Objs) const
-{
-	for (const FNiagaraConstants_DataObject& C : DataObjectConstants)
-	{
-		Objs.Add(C.Value);
-	}
-}
-
 void FNiagaraConstants::AppendToConstantsTable(TArray<FVector4>& ConstantsTable, const FNiagaraConstantMap& Externals)const
 {
 	int32 Idx = ConstantsTable.Num();
@@ -354,15 +266,6 @@ void FNiagaraConstants::AppendToConstantsTable(TArray<FVector4>& ConstantsTable,
 	}
 }
 
-void FNiagaraConstants::AppendExternalBufferConstants(TArray<class UNiagaraDataObject*> &DataObjs, const FNiagaraConstantMap& Externals) const
-{
-	for (int32 i = 0; i < DataObjectConstants.Num(); ++i)
-	{
-		UNiagaraDataObject* Data = Externals.FindDataObj(DataObjectConstants[i].Name);
-		DataObjs.Add(Data);
-	}
-}
-
 void FNiagaraConstants::SetOrAdd(const FNiagaraVariableInfo& Constant, float Value)
 {
 	check(Constant.Type == ENiagaraDataType::Scalar);
@@ -405,20 +308,6 @@ void FNiagaraConstants::SetOrAdd(const FNiagaraVariableInfo& Constant, const FMa
 	}
 }
 
-void FNiagaraConstants::SetOrAdd(const FNiagaraVariableInfo& Constant, UNiagaraDataObject* Value)
-{
-	check(Constant.Type == ENiagaraDataType::Curve);
-	FNiagaraConstants_DataObject* Const = DataObjectConstants.FindByPredicate([&](const FNiagaraConstants_DataObject& C){ return Constant.Name == C.Name; });
-	if (Const)
-	{
-		Const->Value = Value;
-	}
-	else
-	{
-		DataObjectConstants.Add(FNiagaraConstants_DataObject(Constant.Name, Value));
-	}
-}
-
 void FNiagaraConstants::SetOrAdd(FName Name, float Value)
 {
 	SetOrAdd(FNiagaraVariableInfo(Name, ENiagaraDataType::Scalar), Value);
@@ -434,11 +323,6 @@ void FNiagaraConstants::SetOrAdd(FName Name, const FMatrix& Value)
 	SetOrAdd(FNiagaraVariableInfo(Name, ENiagaraDataType::Matrix), Value);
 }
 
-void FNiagaraConstants::SetOrAdd(FName Name, UNiagaraDataObject* Value)
-{
-	SetOrAdd(FNiagaraVariableInfo(Name, ENiagaraDataType::Curve), Value);
-}
-
 float* FNiagaraConstants::FindScalar(FName Name)
 {
 	auto* C = ScalarConstants.FindByPredicate([&](const FNiagaraConstants_Float& Constant){ return Constant.Name == Name; });
@@ -454,10 +338,5 @@ FMatrix* FNiagaraConstants::FindMatrix(FName Name)
 	auto* C = MatrixConstants.FindByPredicate([&](const FNiagaraConstants_Matrix& Constant){ return Constant.Name == Name; });
 	return C ? &C->Value : nullptr;
 }
-UNiagaraDataObject* FNiagaraConstants::FindDataObj(FName Name)
-{
-	auto* C = DataObjectConstants.FindByPredicate([&](const FNiagaraConstants_DataObject& Constant){ return Constant.Name == Name; });
-	return C ? C->Value : nullptr;
-}
 
 //////////////////////////////////////////////////////////////////////////
diff --git a/Engine/Source/Runtime/Niagara/Private/NiagaraEffect.cpp b/Engine/Source/Runtime/Niagara/Private/NiagaraEffect.cpp
--- a/Engine/Source/Runtime/Niagara/Private/NiagaraEffect.cpp
+++ b/Engine/Source/Runtime/Niagara/Private/NiagaraEffect.cpp
@@ -82,8 +82,6 @@ void FNiagaraEffectInstance::DeleteEmitter(TSharedPtr<FNiagaraSimulation> Emitte
 
 void FNiagaraEffectInstance::Tick(float DeltaSeconds)
 {
-	Constants.SetOrAdd(FName(TEXT("EffectGrid")), VolumeGrid);
-
 	// pass the constants down to the emitter
 	// TODO: should probably just pass a pointer to the table
 	EffectBounds.Init();
diff --git a/Engine/Source/Runtime/Niagara/Private/NiagaraSimulation.cpp b/Engine/Source/Runtime/Niagara/Private/NiagaraSimulation.cpp
--- a/Engine/Source/Runtime/Niagara/Private/NiagaraSimulation.cpp
+++ b/Engine/Source/Runtime/Niagara/Private/NiagaraSimulation.cpp
@@ -448,8 +448,7 @@ void FNiagaraSimulation::RunVMScript(FNiagaraEmitterScriptProperties& ScriptProp
 
 	//Fill constant table with required emitter constants and internal script constants.
 	TArray<FVector4> ConstantTable;
-	TArray<UNiagaraDataObject *>DataObjTable;
-	Script->ConstantData.FillConstantTable(Constants, ConstantTable, DataObjTable);
+	Script->ConstantData.FillConstantTable(Constants, ConstantTable);
 
 
 	//Fill in the shared data table
@@ -508,7 +507,6 @@ void FNiagaraSimulation::RunVMScript(FNiagaraEmitterScriptProperties& ScriptProp
 		OutputRegisters,
 		NumOutputRegitsers,
 		ConstantTable.GetData(),
-		DataObjTable.GetData(),
 		SharedDataTable.Num() > 0 ? SharedDataTable.GetData() : NULL,
 		NumParticles
 		);
diff --git a/Engine/Source/Runtime/Niagara/Public/NiagaraCommon.h b/Engine/Source/Runtime/Niagara/Public/NiagaraCommon.h
--- a/Engine/Source/Runtime/Niagara/Public/NiagaraCommon.h
+++ b/Engine/Source/Runtime/Niagara/Public/NiagaraCommon.h
@@ -12,7 +12,6 @@ enum class ENiagaraDataType : uint8
 	Scalar,
 	Vector,
 	Matrix,
-	Curve,
 };
 
 
diff --git a/Engine/Source/Runtime/Niagara/Public/NiagaraEffectRendererProperties.h b/Engine/Source/Runtime/Niagara/Public/NiagaraEffectRendererProperties.h
new file mode 100644
--- /dev/null
+++ b/Engine/Source/Runtime/Niagara/Public/NiagaraEffectRendererProperties.h
@@ -0,0 +1,23 @@
+// Copyright 1998-2016 Epic Games, Inc. All Rights Reserved.
+
+#include "NiagaraEffectRendererProperties.generated.h"
+
+#pragma once
+/**
+* Emitter properties base class
+* Each Effectrenderer derives from this with its own class, and returns it in GetProperties; a copy
+* of those specific properties is stored on UNiagaraEmitterProperties (on the effect) for serialization 
+* and handed back to the effect renderer on load.
+*/
+UCLASS()
+class NIAGARA_API UNiagaraEffectRendererProperties : public UObject
+{
+	GENERATED_UCLASS_BODY()
+	UNiagaraEffectRendererProperties()
+	{}
+
+	UPROPERTY()
+	FName dummy;
+};
+
+
diff --git a/Engine/Source/Runtime/RenderCore/Private/RenderUtils.cpp b/Engine/Source/Runtime/RenderCore/Private/RenderUtils.cpp
--- a/Engine/Source/Runtime/RenderCore/Private/RenderUtils.cpp
+++ b/Engine/Source/Runtime/RenderCore/Private/RenderUtils.cpp
@@ -820,6 +820,14 @@ RENDERCORE_API bool IsSimpleForwardShadingEnabled(EShaderPlatform Platform)
 	return CVar->GetValueOnAnyThread() != 0 && PlatformSupportsSimpleForwardShading(Platform);
 }
 
+RENDERCORE_API bool IsForwardShadingEnabled(ERHIFeatureLevel::Type FeatureLevel)
+{
+	static const auto CVar = IConsoleManager::Get().FindTConsoleVariableDataInt(TEXT("r.ForwardShading"));
+	return CVar->GetValueOnAnyThread() != 0 
+		// Culling uses compute shader
+		&& FeatureLevel >= ERHIFeatureLevel::SM5;
+}
+
 class FUnitCubeVertexBuffer : public FVertexBuffer
 {
 public:
diff --git a/Engine/Source/Runtime/RenderCore/Public/RenderUtils.h b/Engine/Source/Runtime/RenderCore/Public/RenderUtils.h
--- a/Engine/Source/Runtime/RenderCore/Public/RenderUtils.h
+++ b/Engine/Source/Runtime/RenderCore/Public/RenderUtils.h
@@ -313,6 +313,13 @@ RENDERCORE_API bool PlatformSupportsSimpleForwardShading(EShaderPlatform Platfor
 
 RENDERCORE_API bool IsSimpleForwardShadingEnabled(EShaderPlatform Platform);
 
+RENDERCORE_API bool IsForwardShadingEnabled(ERHIFeatureLevel::Type FeatureLevel);
+
+inline bool IsAnyForwardShadingEnabled(EShaderPlatform Platform)
+{
+	return IsForwardShadingEnabled(GetMaxSupportedFeatureLevel(Platform)) || IsSimpleForwardShadingEnabled(Platform);
+}
+
 /** Unit cube vertex buffer (VertexDeclarationFVector4) */
 RENDERCORE_API FVertexBufferRHIRef& GetUnitCubeVertexBuffer();
 
diff --git a/Engine/Source/Runtime/Renderer/Private/BasePassRendering.cpp b/Engine/Source/Runtime/Renderer/Private/BasePassRendering.cpp
--- a/Engine/Source/Runtime/Renderer/Private/BasePassRendering.cpp
+++ b/Engine/Source/Runtime/Renderer/Private/BasePassRendering.cpp
@@ -200,18 +200,17 @@ void FBasePassReflectionParameters::SetMesh(FRHICommandList& RHICmdList, FPixelS
 
 	FMatrix BoxTransformVal = FMatrix::Identity;
 	FVector4 PositionAndRadius = FVector::ZeroVector;
-	FVector4 BoxScalesVal = FVector::ZeroVector;
+	FVector4 BoxScalesVal(0, 0, 0, 0);
 	FVector CaptureOffsetVal = FVector::ZeroVector;
 	EReflectionCaptureShape::Type CaptureShape = EReflectionCaptureShape::Box;
 	
-
 	if (PrimitiveSceneInfo && ReflectionProxy)
 	{
 		PrimitiveSceneInfo->Scene->GetCaptureParameters(ReflectionProxy, CubeArrayTexture, ArrayIndex);
 		PositionAndRadius = FVector4(ReflectionProxy->Position, ReflectionProxy->InfluenceRadius);
 		CaptureShape = ReflectionProxy->Shape;
 		BoxTransformVal = ReflectionProxy->BoxTransform;
-		BoxScalesVal = ReflectionProxy->BoxScales;
+		BoxScalesVal = FVector4(ReflectionProxy->BoxScales, ReflectionProxy->BoxTransitionDistance);
 		CaptureOffsetVal = ReflectionProxy->CaptureOffset;
 	}
 
@@ -451,18 +450,27 @@ public:
 		const typename LightMapPolicyType::ElementDataType& LightMapElementData
 		) const
 	{
-		#if !(UE_BUILD_SHIPPING || UE_BUILD_TEST)
-		// When rendering masked materials in the shader complexity viewmode, 
-		// We want to overwrite complexity for the pixels which get depths written,
-		// And accumulate complexity for pixels which get killed due to the opacity mask being below the clip value.
-		// This is accomplished by forcing the masked materials to render depths in the depth only pass, 
-		// Then rendering in the base pass with additive complexity blending, depth tests on, and depth writes off.
+#if !(UE_BUILD_SHIPPING || UE_BUILD_TEST)
 		if(View.Family->EngineShowFlags.ShaderComplexity)
 		{
+			// When rendering masked materials in the shader complexity viewmode, 
+			// We want to overwrite complexity for the pixels which get depths written,
+			// And accumulate complexity for pixels which get killed due to the opacity mask being below the clip value.
+			// This is accomplished by forcing the masked materials to render depths in the depth only pass, 
+			// Then rendering in the base pass with additive complexity blending, depth tests on, and depth writes off.
 			RHICmdList.SetDepthStencilState(TStaticDepthStencilState<false,CF_DepthNearOrEqual>::GetRHI());
 		}
-
-		const FSceneRenderTargets& SceneContext = FSceneRenderTargets::Get(RHICmdList);
+		else if (View.Family->UseDebugViewPS() && View.Family->GetDebugViewShaderMode() != DVSM_MaterialTexCoordScalesAnalysis)
+		{
+			if (Parameters.PrimitiveSceneProxy && Parameters.PrimitiveSceneProxy->IsSelected())
+			{
+				RHICmdList.SetDepthStencilState(TStaticDepthStencilState<true, CF_DepthNearOrEqual>::GetRHI());
+			}
+			else // If not selected, use depth equal to make alpha test stand out (goes with EarlyZPassMode = DDM_AllOpaque) 
+			{
+				RHICmdList.SetDepthStencilState(TStaticDepthStencilState<false, CF_Equal>::GetRHI());
+			}
+		}
 #endif
 		const FScene* Scene = Parameters.PrimitiveSceneProxy ? Parameters.PrimitiveSceneProxy->GetPrimitiveSceneInfo()->Scene : NULL;
 
diff --git a/Engine/Source/Runtime/Renderer/Private/BasePassRendering.h b/Engine/Source/Runtime/Renderer/Private/BasePassRendering.h
--- a/Engine/Source/Runtime/Renderer/Private/BasePassRendering.h
+++ b/Engine/Source/Runtime/Renderer/Private/BasePassRendering.h
@@ -18,6 +18,129 @@ extern int32 GIndirectLightingCache;
 /** Whether some GBuffer targets are optional. */
 extern bool UseSelectiveBasePassOutputs();
 
+class FForwardLocalLightData
+{
+public:
+	FVector4 LightPositionAndInvRadius;
+	FVector4 LightColorAndFalloffExponent;
+	FVector4 LightDirectionAndShadowMapChannelMask;
+	FVector4 SpotAnglesAndSourceRadius;
+};
+
+/** Parameters for computing forward lighting. */
+class FForwardLightingParameters
+{
+public:
+
+	void Bind(const FShaderParameterMap& ParameterMap)
+	{
+		ForwardGlobalLightData.Bind(ParameterMap, TEXT("ForwardGlobalLightData"));
+		ForwardLocalLightBuffer.Bind(ParameterMap, TEXT("ForwardLocalLightBuffer"));
+		NumCulledLightsGrid.Bind(ParameterMap, TEXT("NumCulledLightsGrid"));
+		CulledLightDataGrid.Bind(ParameterMap, TEXT("CulledLightDataGrid"));
+		LightAttenuationTexture.Bind(ParameterMap, TEXT("LightAttenuationTexture"));
+		LightAttenuationTextureSampler.Bind(ParameterMap, TEXT("LightAttenuationTextureSampler"));
+		IndirectOcclusionTexture.Bind(ParameterMap, TEXT("IndirectOcclusionTexture"));
+		IndirectOcclusionTextureSampler.Bind(ParameterMap, TEXT("IndirectOcclusionTextureSampler"));
+	}
+
+	template<typename ShaderRHIParamRef>
+	void Set(FRHICommandList& RHICmdList, const ShaderRHIParamRef& ShaderRHI, FShader* Shader, const FViewInfo& View)
+	{
+		//@todo - put all of these in a shader resource table
+		SetUniformBufferParameter(RHICmdList, ShaderRHI, ForwardGlobalLightData, View.ForwardLightingResources->ForwardGlobalLightData);
+		SetSRVParameter(RHICmdList, ShaderRHI, ForwardLocalLightBuffer, View.ForwardLightingResources->ForwardLocalLightBuffer.SRV);
+		NumCulledLightsGrid.SetBuffer(RHICmdList, ShaderRHI, View.ForwardLightingResources->NumCulledLightsGrid);
+		CulledLightDataGrid.SetBuffer(RHICmdList, ShaderRHI, View.ForwardLightingResources->CulledLightDataGrid);
+
+		if (LightAttenuationTexture.IsBound() || IndirectOcclusionTexture.IsBound())
+		{
+			FSceneRenderTargets& SceneRenderTargets = FSceneRenderTargets::Get(RHICmdList);
+
+			SetTextureParameter(
+				RHICmdList, 
+				ShaderRHI,
+				LightAttenuationTexture,
+				LightAttenuationTextureSampler,
+				TStaticSamplerState<SF_Point,AM_Wrap,AM_Wrap,AM_Wrap>::GetRHI(),
+				SceneRenderTargets.GetEffectiveLightAttenuationTexture(true)
+				);
+
+			IPooledRenderTarget* IndirectOcclusion = SceneRenderTargets.ScreenSpaceAO;
+
+			if (!SceneRenderTargets.bScreenSpaceAOIsValid)
+			{
+				IndirectOcclusion = GSystemTextures.WhiteDummy;
+			}
+
+			SetTextureParameter(
+				RHICmdList, 
+				ShaderRHI,
+				IndirectOcclusionTexture,
+				IndirectOcclusionTextureSampler,
+				TStaticSamplerState<SF_Point,AM_Wrap,AM_Wrap,AM_Wrap>::GetRHI(),
+				IndirectOcclusion->GetRenderTargetItem().ShaderResourceTexture
+				);
+		}
+	}
+
+	template<typename ShaderRHIParamRef>
+	void UnsetParameters(FRHICommandList& RHICmdList, const ShaderRHIParamRef& ShaderRHI, const FViewInfo& View)
+	{
+		NumCulledLightsGrid.UnsetUAV(RHICmdList, ShaderRHI);
+		CulledLightDataGrid.UnsetUAV(RHICmdList, ShaderRHI);
+
+		TArray<FUnorderedAccessViewRHIParamRef, TInlineAllocator<2>> OutUAVs;
+
+		if (NumCulledLightsGrid.IsBound())
+		{
+			OutUAVs.Add(View.ForwardLightingResources->NumCulledLightsGrid.UAV);
+		}
+
+		if (CulledLightDataGrid.IsBound())
+		{
+			OutUAVs.Add(View.ForwardLightingResources->CulledLightDataGrid.UAV);
+		}
+
+		if (OutUAVs.Num() > 0)
+		{
+			RHICmdList.TransitionResources(EResourceTransitionAccess::EReadable, EResourceTransitionPipeline::EComputeToGfx, OutUAVs.GetData(), OutUAVs.Num());
+		}
+	}
+
+	static void ModifyCompilationEnvironment(EShaderPlatform Platform, FShaderCompilerEnvironment& OutEnvironment)
+	{
+		OutEnvironment.SetDefine(TEXT("LOCAL_LIGHT_DATA_STRIDE"), FMath::DivideAndRoundUp<int32>(sizeof(FForwardLocalLightData), sizeof(FVector4)));
+		extern int32 NumCulledLightsGridStride;
+		OutEnvironment.SetDefine(TEXT("NUM_CULLED_LIGHTS_GRID_STRIDE"), NumCulledLightsGridStride);
+	}
+
+	/** Serializer. */
+	friend FArchive& operator<<(FArchive& Ar,FForwardLightingParameters& P)
+	{
+		Ar << P.ForwardGlobalLightData;
+		Ar << P.ForwardLocalLightBuffer;
+		Ar << P.NumCulledLightsGrid;
+		Ar << P.CulledLightDataGrid;
+		Ar << P.LightAttenuationTexture;
+		Ar << P.LightAttenuationTextureSampler;
+		Ar << P.IndirectOcclusionTexture;
+		Ar << P.IndirectOcclusionTextureSampler;
+		return Ar;
+	}
+
+private:
+
+	FShaderUniformBufferParameter ForwardGlobalLightData;
+	FShaderResourceParameter ForwardLocalLightBuffer;
+	FRWShaderParameter NumCulledLightsGrid;
+	FRWShaderParameter CulledLightDataGrid;
+	FShaderResourceParameter LightAttenuationTexture;
+	FShaderResourceParameter LightAttenuationTextureSampler;
+	FShaderResourceParameter IndirectOcclusionTexture;
+	FShaderResourceParameter IndirectOcclusionTextureSampler;
+};
+
 /** Parameters needed for looking up into translucency lighting volumes. */
 class FTranslucentLightingVolumeParameters
 {
@@ -515,6 +638,8 @@ public:
 			const int32 VelocityIndex = 4; // As defined in BasePassPixelShader.usf
 			OutEnvironment.SetRenderTargetOutputFormat(VelocityIndex, PF_G16R16);
 		}
+
+		FForwardLightingParameters::ModifyCompilationEnvironment(Platform, OutEnvironment);
 	}
 
 	/** Initialization constructor. */
@@ -525,8 +650,8 @@ public:
 		ReflectionParameters.Bind(Initializer.ParameterMap);
 		TranslucentLightingParameters.Bind(Initializer.ParameterMap);
 		EditorCompositeParams.Bind(Initializer.ParameterMap);
-		LightGrid.Bind(Initializer.ParameterMap,TEXT("LightGrid"));
 		DownsampleFactorFromSceneBufferSize.Bind(Initializer.ParameterMap, TEXT("DownsampleFactorFromSceneBufferSize"));
+		ForwardLightingParameters.Bind(Initializer.ParameterMap);
 	}
 	TBasePassPixelShaderPolicyParamType() {}
 
@@ -544,18 +669,17 @@ public:
 
 		FMeshMaterialShader::SetParameters(RHICmdList, ShaderRHI, MaterialRenderProxy, MaterialResource, *View, TextureMode);
 
-		if (View->GetFeatureLevel() >= ERHIFeatureLevel::SM4)
-		{
-			ReflectionParameters.Set(RHICmdList, ShaderRHI, View);
+		ReflectionParameters.Set(RHICmdList, ShaderRHI, View);
 
-			if (IsTranslucentBlendMode(BlendMode))
-			{
-				SetShaderValue(RHICmdList, ShaderRHI, DownsampleFactorFromSceneBufferSize, DownsampleFactorFromSceneBufferSizeValue);
-				TranslucentLightingParameters.Set(RHICmdList, ShaderRHI, View);
-			}
+		if (IsTranslucentBlendMode(BlendMode))
+		{
+			SetShaderValue(RHICmdList, ShaderRHI, DownsampleFactorFromSceneBufferSize, DownsampleFactorFromSceneBufferSizeValue);
+			TranslucentLightingParameters.Set(RHICmdList, ShaderRHI, View);
 		}
 		
 		EditorCompositeParams.SetParameters(RHICmdList, MaterialResource, View, bEnableEditorPrimitveDepthTest, GetPixelShader());
+
+		ForwardLightingParameters.Set(RHICmdList, ShaderRHI, this, *View);
 	}
 
 	void SetMesh(FRHICommandList& RHICmdList, const FVertexFactory* VertexFactory,const FSceneView& View,const FPrimitiveSceneProxy* Proxy,const FMeshBatchElement& BatchElement, const FMeshDrawingRenderState& DrawRenderState, EBlendMode BlendMode);
@@ -567,8 +691,8 @@ public:
 		Ar << ReflectionParameters;
 		Ar << TranslucentLightingParameters;
  		Ar << EditorCompositeParams;
-		Ar << LightGrid;
 		Ar << DownsampleFactorFromSceneBufferSize;
+		Ar << ForwardLightingParameters;
 		return bShaderHasOutdatedParameters;
 	}
 
@@ -576,8 +700,8 @@ private:
 	FBasePassReflectionParameters ReflectionParameters;
 	FTranslucentLightingParameters TranslucentLightingParameters;
 	FEditorCompositingParameters EditorCompositeParams;
-	FShaderResourceParameter LightGrid;
 	FShaderParameter DownsampleFactorFromSceneBufferSize;
+	FForwardLightingParameters ForwardLightingParameters;
 };
 
 /**
@@ -831,12 +955,17 @@ public:
 
 		if (UseDebugViewPS())
 		{
-			// If we are in the translucent pass then override the blend mode, otherwise maintain additive blending.
-			if (View->Family->EngineShowFlags.ShaderComplexity && IsTranslucentBlendMode(BlendMode))
+			if (IsTranslucentBlendMode(BlendMode))
 			{
-				RHICmdList.SetBlendState( TStaticBlendState<CW_RGBA, BO_Add, BF_One, BF_One, BO_Add, BF_Zero, BF_One>::GetRHI());
+				if (View->Family->EngineShowFlags.ShaderComplexity)
+				{	// If we are in the translucent pass then override the blend mode, otherwise maintain additive blending.
+					RHICmdList.SetBlendState(TStaticBlendState<CW_RGBA, BO_Add, BF_One, BF_One, BO_Add, BF_Zero, BF_One>::GetRHI());
+				}
+				else if (View->Family->GetDebugViewShaderMode() != DVSM_MaterialTexCoordScalesAnalysis)
+				{	// Otherwise, force translucent blend mode (shaders will use an hardcoded alpha).
+					RHICmdList.SetBlendState(TStaticBlendState<CW_RGBA, BO_Add, BF_SourceAlpha, BF_InverseSourceAlpha, BO_Add, BF_Zero, BF_InverseSourceAlpha>::GetRHI());
+				}
 			}
-
 			FDebugViewMode::GetPSInterface(View->ShaderMap, MaterialResource, GetDebugViewShaderMode())->SetParameters(RHICmdList, VertexShader, PixelShader, MaterialRenderProxy, *MaterialResource, *View);
 		}
 		else
@@ -960,7 +1089,7 @@ public:
 			PixelShader->SetMesh(RHICmdList, VertexFactory,View,PrimitiveSceneProxy,BatchElement,DrawRenderState,BlendMode);
 		}
 
-		if (bEnableReceiveDecalOutput)
+		if (bEnableReceiveDecalOutput && !UseDebugViewPS())
 		{
 			// Set stencil value for this draw call
 			// This is effectively extending the GBuffer using the stencil bits
diff --git a/Engine/Source/Runtime/Renderer/Private/CapsuleShadowRendering.cpp b/Engine/Source/Runtime/Renderer/Private/CapsuleShadowRendering.cpp
--- a/Engine/Source/Runtime/Renderer/Private/CapsuleShadowRendering.cpp
+++ b/Engine/Source/Runtime/Renderer/Private/CapsuleShadowRendering.cpp
@@ -611,7 +611,8 @@ bool SupportsCapsuleShadows(ERHIFeatureLevel::Type FeatureLevel, EShaderPlatform
 bool FDeferredShadingSceneRenderer::RenderCapsuleDirectShadows(
 	const FLightSceneInfo& LightSceneInfo,
 	FRHICommandListImmediate& RHICmdList, 
-	const TArray<FProjectedShadowInfo*, SceneRenderingAllocator>& CapsuleShadows) const
+	const TArray<FProjectedShadowInfo*, SceneRenderingAllocator>& CapsuleShadows, 
+	bool bProjectingForForwardShading) const
 {
 	bool bAllViewsHaveViewState = true;
 
@@ -767,7 +768,7 @@ bool FDeferredShadingSceneRenderer::RenderCapsuleDirectShadows(
 				}
 
 				{
-					SCOPED_DRAW_EVENTF(RHICmdList, Upsample, TEXT("UpsampleDirectCapsuleShadowRendering %dx%d"),
+					SCOPED_DRAW_EVENTF(RHICmdList, Upsample, TEXT("Upsample %dx%d"),
 						ScissorRect.Width(), ScissorRect.Height());
 						
 					FSceneRenderTargets::Get(RHICmdList).BeginRenderingLightAttenuation(RHICmdList);
@@ -776,9 +777,13 @@ bool FDeferredShadingSceneRenderer::RenderCapsuleDirectShadows(
 					RHICmdList.SetRasterizerState(TStaticRasterizerState<FM_Solid, CM_None>::GetRHI());
 					RHICmdList.SetDepthStencilState(TStaticDepthStencilState<false, CF_Always>::GetRHI());
 				
-					// use B and A in Light Attenuation for per-object shadows
-					// CO_Min is needed to combine multiple shadow passes
-					RHICmdList.SetBlendState(TStaticBlendState<CW_BA, BO_Min, BF_One, BF_One, BO_Min, BF_One, BF_One>::GetRHI());
+					FProjectedShadowInfo::SetBlendStateForProjection(
+						RHICmdList,
+						LightSceneInfo.Proxy->GetPreviewShadowMapChannel(),
+						false,
+						false,
+						bProjectingForForwardShading,
+						false);
 
 					TShaderMapRef<FCapsuleShadowingUpsampleVS> VertexShader(View.ShaderMap);
 
@@ -988,7 +993,10 @@ void FDeferredShadingSceneRenderer::SetupIndirectCapsuleShadows(FRHICommandListI
 	NumCapsuleShapes = CapsuleShapeData.Num();
 }
 
-void FDeferredShadingSceneRenderer::RenderIndirectCapsuleShadows(FRHICommandListImmediate& RHICmdList) const
+void FDeferredShadingSceneRenderer::RenderIndirectCapsuleShadows(
+	FRHICommandListImmediate& RHICmdList, 
+	FTextureRHIParamRef IndirectLightingTexture, 
+	FTextureRHIParamRef ExistingIndirectOcclusionTexture) const
 {
 	if (SupportsCapsuleShadows(FeatureLevel, GShaderPlatformForFeatureLevel[FeatureLevel])
 		&& FSceneRenderTargets::Get(RHICmdList).IsStaticLightingAllowed())
@@ -1009,8 +1017,6 @@ void FDeferredShadingSceneRenderer::RenderIndirectCapsuleShadows(FRHICommandList
 
 		if (bAnyViewsUseCapsuleShadows)
 		{
-			FSceneRenderTargets::Get(RHICmdList).FinishRenderingSceneColor(RHICmdList);
-
 			TRefCountPtr<IPooledRenderTarget> RayTracedShadowsRT;
 
 			{
@@ -1020,6 +1026,32 @@ void FDeferredShadingSceneRenderer::RenderIndirectCapsuleShadows(FRHICommandList
 				GRenderTargetPool.FindFreeElement(RHICmdList, Desc, RayTracedShadowsRT, TEXT("RayTracedShadows"));
 			}
 
+			FSceneRenderTargets& SceneContext = FSceneRenderTargets::Get(RHICmdList);
+			TArray<FTextureRHIParamRef, TInlineAllocator<2>> RenderTargets;
+
+			if (IndirectLightingTexture)
+			{
+				RenderTargets.Add(IndirectLightingTexture);
+			}
+
+			if (ExistingIndirectOcclusionTexture)
+			{
+				RenderTargets.Add(ExistingIndirectOcclusionTexture);
+			}
+
+			if (RenderTargets.Num() == 0)
+			{
+				SceneContext.bScreenSpaceAOIsValid = true;
+				RenderTargets.Add(SceneContext.ScreenSpaceAO->GetRenderTargetItem().TargetableTexture);
+
+				SCOPED_DRAW_EVENT(RHICmdList, ClearIndirectOcclusion);
+				// We are the first users of the indirect occlusion texture so we must clear to unoccluded
+				SetRenderTargets(RHICmdList, RenderTargets.Num(), RenderTargets.GetData(), FTextureRHIParamRef(), 0, NULL, true);
+				RHICmdList.Clear(true, FLinearColor::White, false, 0, false, 0, FIntRect());
+			}
+							
+			check(RenderTargets.Num() > 0);
+
 			for (int32 ViewIndex = 0; ViewIndex < Views.Num(); ViewIndex++)
 			{
 				const FViewInfo& View = Views[ViewIndex];
@@ -1076,19 +1108,9 @@ void FDeferredShadingSceneRenderer::RenderIndirectCapsuleShadows(FRHICommandList
 						}
 			
 						{
-							SCOPED_DRAW_EVENTF(RHICmdList, Upsample, TEXT("UpsampleIndirectCapsuleShadowRendering %dx%d"),
-								ScissorRect.Width(), ScissorRect.Height());
-
-							FSceneRenderTargets& SceneContext = FSceneRenderTargets::Get(RHICmdList);
-
-							FTextureRHIParamRef RenderTargets[2] =
-							{
-								SceneContext.GetSceneColorSurface(),
-								SceneContext.bScreenSpaceAOIsValid ? SceneContext.ScreenSpaceAO->GetRenderTargetItem().TargetableTexture : NULL
-							};
+							SCOPED_DRAW_EVENTF(RHICmdList, Upsample, TEXT("Upsample %dx%d"), ScissorRect.Width(), ScissorRect.Height());
 
-							const int32 NumTargets = ARRAY_COUNT(RenderTargets) - (SceneContext.bScreenSpaceAOIsValid ? 0 : 1);
-							SetRenderTargets(RHICmdList, NumTargets, RenderTargets, FTextureRHIParamRef(), 0, NULL, true);
+							SetRenderTargets(RHICmdList, RenderTargets.Num(), RenderTargets.GetData(), FTextureRHIParamRef(), 0, NULL, true);
 
 							RHICmdList.SetViewport(View.ViewRect.Min.X, View.ViewRect.Min.Y, 0.0f, View.ViewRect.Max.X, View.ViewRect.Max.Y, 1.0f);
 							RHICmdList.SetRasterizerState(TStaticRasterizerState<FM_Solid, CM_None>::GetRHI());
@@ -1096,7 +1118,7 @@ void FDeferredShadingSceneRenderer::RenderIndirectCapsuleShadows(FRHICommandList
 				
 							// Modulative blending against scene color for application to indirect diffuse
 							// Modulative blending against SSAO occlusion value for application to indirect specular, since Reflection Environment pass masks by AO
-							if (SceneContext.bScreenSpaceAOIsValid)
+							if (RenderTargets.Num() > 1)
 							{
 								RHICmdList.SetBlendState(TStaticBlendState<
 									CW_RGB, BO_Add, BF_DestColor, BF_Zero, BO_Add, BF_Zero, BF_One,
@@ -1109,7 +1131,7 @@ void FDeferredShadingSceneRenderer::RenderIndirectCapsuleShadows(FRHICommandList
 
 							TShaderMapRef<FCapsuleShadowingUpsampleVS> VertexShader(View.ShaderMap);
 
-							if (SceneContext.bScreenSpaceAOIsValid)
+							if (RenderTargets.Num() > 1)
 							{
 								if (GCapsuleShadowsFullResolution)
 								{
diff --git a/Engine/Source/Runtime/Renderer/Private/CompositionLighting/CompositionLighting.cpp b/Engine/Source/Runtime/Renderer/Private/CompositionLighting/CompositionLighting.cpp
--- a/Engine/Source/Runtime/Renderer/Private/CompositionLighting/CompositionLighting.cpp
+++ b/Engine/Source/Runtime/Renderer/Private/CompositionLighting/CompositionLighting.cpp
@@ -58,7 +58,7 @@ bool IsAmbientCubemapPassRequired(const FSceneView& View)
 {
 	FScene* Scene = (FScene*)View.Family->Scene;
 
-	return View.FinalPostProcessSettings.ContributingCubemaps.Num() != 0 && !IsSimpleForwardShadingEnabled(View.GetShaderPlatform());
+	return View.FinalPostProcessSettings.ContributingCubemaps.Num() != 0 && !IsAnyForwardShadingEnabled(View.GetShaderPlatform());
 }
 
 bool IsLpvIndirectPassRequired(const FViewInfo& View)
@@ -98,7 +98,7 @@ static bool IsReflectionEnvironmentActive(const FSceneView& View)
 	bool HasReflectionCaptures = (Scene->ReflectionSceneData.RegisteredReflectionCaptures.Num() > 0);
 	bool HasSSR = View.Family->EngineShowFlags.ScreenSpaceReflections;
 
-	return (Scene->GetFeatureLevel() == ERHIFeatureLevel::SM5 && IsReflectingEnvironment && (HasReflectionCaptures || HasSSR) && !IsSimpleForwardShadingEnabled(View.GetShaderPlatform()));
+	return (Scene->GetFeatureLevel() == ERHIFeatureLevel::SM5 && IsReflectingEnvironment && (HasReflectionCaptures || HasSSR) && !IsAnyForwardShadingEnabled(View.GetShaderPlatform()));
 }
 
 static bool IsSkylightActive(const FViewInfo& View)
@@ -123,7 +123,7 @@ uint32 ComputeAmbientOcclusionPassCount(const FViewInfo& View)
 			&& View.FinalPostProcessSettings.AmbientOcclusionRadius >= 0.1f
 			&& !View.Family->UseDebugViewPS()
 			&& (FSSAOHelper::IsBasePassAmbientOcclusionRequired(View) || IsAmbientCubemapPassRequired(View) || IsReflectionEnvironmentActive(View) || IsSkylightActive(View) || View.Family->EngineShowFlags.VisualizeBuffer)
-			&& !IsSimpleForwardShadingEnabled(View.GetShaderPlatform());
+			&& !IsAnyForwardShadingEnabled(View.GetShaderPlatform());
 	}
 
 	if (bEnabled)
@@ -414,7 +414,7 @@ void FCompositionLighting::ProcessAfterLighting(FRHICommandListImmediate& RHICmd
 		// Screen Space Subsurface Scattering
 		{
 			float Radius = CVarSSSScale.GetValueOnRenderThread();
-			bool bSimpleDynamicLighting = IsSimpleForwardShadingEnabled(View.GetShaderPlatform());
+			bool bSimpleDynamicLighting = IsAnyForwardShadingEnabled(View.GetShaderPlatform());
 			bool bScreenSpaceSubsurfacePassNeeded = (View.ShadingModelMaskInView & (1 << MSM_SubsurfaceProfile)) != 0;
 			bool bSubsurfaceAllowed = CVarSubsurfaceScattering.GetValueOnRenderThread() == 1;
 
diff --git a/Engine/Source/Runtime/Renderer/Private/CompositionLighting/PostProcessAmbientOcclusion.cpp b/Engine/Source/Runtime/Renderer/Private/CompositionLighting/PostProcessAmbientOcclusion.cpp
--- a/Engine/Source/Runtime/Renderer/Private/CompositionLighting/PostProcessAmbientOcclusion.cpp
+++ b/Engine/Source/Runtime/Renderer/Private/CompositionLighting/PostProcessAmbientOcclusion.cpp
@@ -120,7 +120,7 @@ EAsyncComputeBudget FSSAOHelper::GetAmbientOcclusionAsyncComputeBudget()
 bool FSSAOHelper::IsBasePassAmbientOcclusionRequired(const FViewInfo& View)
 {
 	// the BaseAO pass is only worth with some AO
-	return (View.FinalPostProcessSettings.AmbientOcclusionStaticFraction >= 1 / 100.0f) && !IsSimpleForwardShadingEnabled(View.GetShaderPlatform());
+	return (View.FinalPostProcessSettings.AmbientOcclusionStaticFraction >= 1 / 100.0f) && !IsAnyForwardShadingEnabled(View.GetShaderPlatform());
 }
 
 bool FSSAOHelper::IsAmbientOcclusionAsyncCompute(const FViewInfo& View, uint32 AOPassCount)
diff --git a/Engine/Source/Runtime/Renderer/Private/DeferredShadingRenderer.cpp b/Engine/Source/Runtime/Renderer/Private/DeferredShadingRenderer.cpp
--- a/Engine/Source/Runtime/Renderer/Private/DeferredShadingRenderer.cpp
+++ b/Engine/Source/Runtime/Renderer/Private/DeferredShadingRenderer.cpp
@@ -89,6 +89,14 @@ static TAutoConsoleVariable<int32> CVarFXSystemPreRenderAfterPrepass(
 	ECVF_RenderThreadSafe
 	);
 
+int32 GbEnableAsyncComputeTranslucencyLightingVolumeClear = 1;
+static FAutoConsoleVariableRef CVarEnableAsyncComputeTranslucencyLightingVolumeClear(
+	TEXT("r.EnableAsyncComputeTranslucencyLightingVolumeClear"),
+	GbEnableAsyncComputeTranslucencyLightingVolumeClear,
+	TEXT("Whether to clear the translucency lighting volume using async compute.\n"),
+	ECVF_RenderThreadSafe | ECVF_Scalability
+);
+
 DECLARE_CYCLE_STAT(TEXT("PostInitViews FlushDel"), STAT_PostInitViews_FlushDel, STATGROUP_InitViews);
 DECLARE_CYCLE_STAT(TEXT("InitViews Intentional Stall"), STAT_InitViews_Intentional_Stall, STATGROUP_InitViews);
 
@@ -120,6 +128,40 @@ DECLARE_CYCLE_STAT(TEXT("DeferredShadingSceneRenderer RenderFinish"), STAT_FDefe
 DECLARE_CYCLE_STAT(TEXT("OcclusionSubmittedFence Dispatch"), STAT_OcclusionSubmittedFence_Dispatch, STATGROUP_SceneRendering);
 DECLARE_CYCLE_STAT(TEXT("OcclusionSubmittedFence Wait"), STAT_OcclusionSubmittedFence_Wait, STATGROUP_SceneRendering);
 
+bool ShouldForceFullDepthPass(ERHIFeatureLevel::Type FeatureLevel)
+{
+	static IConsoleVariable* CDBufferVar = IConsoleManager::Get().FindConsoleVariable(TEXT("r.DBuffer"));
+	bool bDBufferAllowed = CDBufferVar ? CDBufferVar->GetInt() != 0 : false;
+
+	static const auto StencilLODDitherCVar = IConsoleManager::Get().FindTConsoleVariableDataInt(TEXT("r.StencilForLODDither"));
+	bool bStencilLODDither = StencilLODDitherCVar->GetValueOnAnyThread() != 0;
+
+	// Note: ShouldForceFullDepthPass affects which static draw lists meshes go into, so nothing it depends on can change at runtime, unless you do a FGlobalComponentRecreateRenderStateContext to propagate the cvar change
+	return bDBufferAllowed || bStencilLODDither || IsForwardShadingEnabled(FeatureLevel);
+}
+
+const TCHAR* GetDepthPassReason(bool bDitheredLODTransitionsUseStencil, ERHIFeatureLevel::Type FeatureLevel)
+{
+	if (IsForwardShadingEnabled(FeatureLevel))
+	{
+		return TEXT("(Forced by ForwardShading)");
+	}
+
+	static IConsoleVariable* CDBufferVar = IConsoleManager::Get().FindConsoleVariable(TEXT("r.DBuffer"));
+	bool bDBufferAllowed = CDBufferVar ? CDBufferVar->GetInt() != 0 : false;
+
+	if (bDBufferAllowed)
+	{
+		return TEXT("(Forced by DBuffer)");
+	}
+
+	if (bDitheredLODTransitionsUseStencil)
+	{
+		return TEXT("(Forced by StencilLODDither)");
+	}
+
+	return TEXT("");
+}
 
 /*-----------------------------------------------------------------------------
 	FDeferredShadingSceneRenderer
@@ -154,7 +196,7 @@ FDeferredShadingSceneRenderer::FDeferredShadingSceneRenderer(const FSceneViewFam
 	}
 
 	// Shader complexity requires depth only pass to display masked material cost correctly
-	if (ViewFamily.EngineShowFlags.ShaderComplexity)
+	if (ViewFamily.UseDebugViewPS() && ViewFamily.GetDebugViewShaderMode() != DVSM_MaterialTexCoordScalesAnalysis)
 	{
 		EarlyZPassMode = DDM_AllOpaque;
 	}
@@ -163,13 +205,10 @@ FDeferredShadingSceneRenderer::FDeferredShadingSceneRenderer(const FSceneViewFam
 
 	bEarlyZPassMovable = GEarlyZPassMovable != 0;
 
-	static IConsoleVariable* CDBufferVar = IConsoleManager::Get().FindConsoleVariable(TEXT("r.DBuffer"));
-	bool bDBufferAllowed = CDBufferVar ? CDBufferVar->GetInt() != 0 : false;
-
 	static const auto StencilLODDitherCVar = IConsoleManager::Get().FindTConsoleVariableDataInt(TEXT("r.StencilForLODDither"));
 	bDitheredLODTransitionsUseStencil = StencilLODDitherCVar->GetValueOnAnyThread() != 0;
 
-	if (bDBufferAllowed || bDitheredLODTransitionsUseStencil)
+	if (ShouldForceFullDepthPass(FeatureLevel))
 	{
 		// DBuffer decals and stencil LOD dithering force a full prepass
 		EarlyZPassMode = DDM_AllOccluders;
@@ -570,7 +609,7 @@ void FDeferredShadingSceneRenderer::Render(FRHICommandListImmediate& RHICmdList)
 	bool bRequiresRHIClear = true;
 	bool bRequiresFarZQuadClear = false;
 
-	bool bUseGBuffer = !IsSimpleForwardShadingEnabled(GetFeatureLevelShaderPlatform(FeatureLevel));
+	const bool bUseGBuffer = !IsAnyForwardShadingEnabled(GetFeatureLevelShaderPlatform(FeatureLevel));
 
 	if (ClearMethodCVar)
 	{
@@ -712,7 +751,7 @@ void FDeferredShadingSceneRenderer::Render(FRHICommandListImmediate& RHICmdList)
 	}
 
 	//occlusion can't run before basepass if there's no prepass to fill in some depth to occlude against.
-	bool bOcclusionBeforeBasePass = (CVarOcclusionQueryLocation.GetValueOnRenderThread() == 1) && bNeedsPrePass;	
+	bool bOcclusionBeforeBasePass = ((CVarOcclusionQueryLocation.GetValueOnRenderThread() == 1) && bNeedsPrePass) || IsForwardShadingEnabled(FeatureLevel);	
 	bool bHZBBeforeBasePass = bOcclusionBeforeBasePass && (EarlyZPassMode == EDepthDrawingMode::DDM_AllOccluders);
 
 	RenderOcclusion(RHICmdList, bOcclusionBeforeBasePass, bHZBBeforeBasePass);
@@ -738,6 +777,18 @@ void FDeferredShadingSceneRenderer::Render(FRHICommandListImmediate& RHICmdList)
 		RenderCustomDepthPassAtLocation(RHICmdList, 0);
 	}
 
+	ComputeLightGrid(RHICmdList);
+
+	if (IsForwardShadingEnabled(FeatureLevel))
+	{
+		RenderForwardShadingShadowProjections(RHICmdList);
+
+		RenderIndirectCapsuleShadows(
+			RHICmdList, 
+			NULL, 
+			NULL);
+	}
+
 	// only temporarily available after early z pass and until base pass
 	check(!SceneContext.DBufferA);
 	check(!SceneContext.DBufferB);
@@ -787,6 +838,11 @@ void FDeferredShadingSceneRenderer::Render(FRHICommandListImmediate& RHICmdList)
 		SceneContext.BeginRenderingGBuffer(RHICmdList, ERenderTargetLoadAction::ENoAction, DepthLoadAction, ViewFamily.EngineShowFlags.ShaderComplexity);
 	}
 
+	if (GbEnableAsyncComputeTranslucencyLightingVolumeClear && GSupportsEfficientAsyncCompute)
+	{
+		ClearTranslucentVolumeLightingAsyncCompute(RHICmdList);
+	}
+
 	GRenderTargetPool.AddPhaseEvent(TEXT("BasePass"));
 
 	RHICmdList.SetCurrentStat(GET_STATID(STAT_CLM_BasePass));
@@ -884,6 +940,10 @@ void FDeferredShadingSceneRenderer::Render(FRHICommandListImmediate& RHICmdList)
 	// Copy lighting channels out of stencil before deferred decals which overwrite those values
 	CopyStencilToLightingChannelTexture(RHICmdList);
 
+	{
+		GCompositionLighting.GfxWaitForAsyncSSAO(RHICmdList);
+	}
+
 	// Pre-lighting composition lighting stage
 	// e.g. deferred decals, SSAO
 	if (FeatureLevel >= ERHIFeatureLevel::SM4)
@@ -918,10 +978,6 @@ void FDeferredShadingSceneRenderer::Render(FRHICommandListImmediate& RHICmdList)
 		RHICmdList.TransitionResource(EResourceTransitionAccess::EReadable, SceneContext.GetSceneDepthSurface());
 	}
 
-	{
-		GCompositionLighting.GfxWaitForAsyncSSAO(RHICmdList);
-	}
-
 	// Render lighting.
 	if (ViewFamily.EngineShowFlags.Lighting
 		&& FeatureLevel >= ERHIFeatureLevel::SM4
@@ -933,14 +989,20 @@ void FDeferredShadingSceneRenderer::Render(FRHICommandListImmediate& RHICmdList)
 		GRenderTargetPool.AddPhaseEvent(TEXT("Lighting"));
 
 		// These modulate the scenecolor output from the basepass, which is assumed to be indirect lighting
-		RenderIndirectCapsuleShadows(RHICmdList);
+		RenderIndirectCapsuleShadows(
+			RHICmdList, 
+			SceneContext.GetSceneColorSurface(), 
+			SceneContext.bScreenSpaceAOIsValid ? SceneContext.ScreenSpaceAO->GetRenderTargetItem().TargetableTexture : NULL);
 
 		TRefCountPtr<IPooledRenderTarget> DynamicBentNormalAO;
 		// These modulate the scenecolor output from the basepass, which is assumed to be indirect lighting
 		RenderDFAOAsIndirectShadowing(RHICmdList, VelocityRT, DynamicBentNormalAO);
 
 		// Clear the translucent lighting volumes before we accumulate
-		ClearTranslucentVolumeLighting(RHICmdList);
+		if ((GbEnableAsyncComputeTranslucencyLightingVolumeClear && GSupportsEfficientAsyncCompute) == false)
+		{
+			ClearTranslucentVolumeLighting(RHICmdList);
+		}
 
 		RHICmdList.SetCurrentStat(GET_STATID(STAT_CLM_Lighting));
 		RenderLights(RHICmdList);
diff --git a/Engine/Source/Runtime/Renderer/Private/DeferredShadingRenderer.h b/Engine/Source/Runtime/Renderer/Private/DeferredShadingRenderer.h
--- a/Engine/Source/Runtime/Renderer/Private/DeferredShadingRenderer.h
+++ b/Engine/Source/Runtime/Renderer/Private/DeferredShadingRenderer.h
@@ -26,6 +26,8 @@ public:
 	EDepthDrawingMode EarlyZPassMode;
 	bool bEarlyZPassMovable;
 	bool bDitheredLODTransitionsUseStencil;
+	
+	FComputeFenceRHIRef TranslucencyLightingVolumeClearEndFence;
 
 	FDeferredShadingSceneRenderer(const FSceneViewFamily* InViewFamily,FHitProxyConsumer* HitProxyConsumer);
 
@@ -59,6 +61,9 @@ public:
 	 */
 	bool RenderPrePassViewParallel(const FViewInfo& View, FRHICommandListImmediate& ParentCmdList, TFunctionRef<void()> AfterTasksAreStarted, bool bDoPrePre);
 
+	/** Culls local lights to a grid in frustum space.  Needed for forward shading or translucency using the Surface lighting mode. */
+	void ComputeLightGrid(FRHICommandListImmediate& RHICmdList);
+
 	/** Renders the basepass for the static data of a given View. */
 	bool RenderBasePassStaticData(FRHICommandList& RHICmdList, FViewInfo& View);
 	bool RenderBasePassStaticDataType(FRHICommandList& RHICmdList, FViewInfo& View, const EBasePassDrawListType DrawType);
@@ -295,13 +300,17 @@ private:
 	bool RenderCapsuleDirectShadows(
 		const FLightSceneInfo& LightSceneInfo,
 		FRHICommandListImmediate& RHICmdList, 
-		const TArray<FProjectedShadowInfo*, SceneRenderingAllocator>& CapsuleShadows) const;
+		const TArray<FProjectedShadowInfo*, SceneRenderingAllocator>& CapsuleShadows, 
+		bool bProjectingForForwardShading) const;
 
 	/** Sets up ViewState buffers for rendering capsule shadows. */
 	void SetupIndirectCapsuleShadows(FRHICommandListImmediate& RHICmdList, const FViewInfo& View, bool bPrepareLightData, int32& NumCapsuleShapes) const;
 
 	/** Renders indirect shadows from capsules modulated onto scene color. */
-	void RenderIndirectCapsuleShadows(FRHICommandListImmediate& RHICmdList) const;
+	void RenderIndirectCapsuleShadows(
+		FRHICommandListImmediate& RHICmdList, 
+		FTextureRHIParamRef IndirectLightingTexture, 
+		FTextureRHIParamRef ExistingIndirectOcclusionTexture) const;
 
 	/** Renders capsule shadows for movable skylights, using the cone of visibility (bent normal) from DFAO. */
 	void RenderCapsuleShadowsForMovableSkylight(FRHICommandListImmediate& RHICmdList, TRefCountPtr<IPooledRenderTarget>& BentNormalOutput) const;
@@ -309,6 +318,9 @@ private:
 	/** Render deferred projections of shadows for a given light into the light attenuation buffer. */
 	bool RenderShadowProjections(FRHICommandListImmediate& RHICmdList, const FLightSceneInfo* LightSceneInfo, bool& bInjectedTranslucentVolume);
 
+	/** Render shadow projections when forward rendering. */
+	void RenderForwardShadingShadowProjections(FRHICommandListImmediate& RHICmdList);
+
 	/**
 	  * Used by RenderLights to render a light function to the attenuation buffer.
 	  *
@@ -338,6 +350,9 @@ private:
 	/** Clears the translucency lighting volumes before light accumulation. */
 	void ClearTranslucentVolumeLighting(FRHICommandListImmediate& RHICmdList);
 
+	/** Clears the translucency lighting volume via an async compute shader overlapped with the basepass. */
+	void ClearTranslucentVolumeLightingAsyncCompute(FRHICommandListImmediate& RHICmdList);
+
 	/** Add AmbientCubemap to the lighting volumes. */
 	void InjectAmbientCubemapTranslucentVolumeLighting(FRHICommandList& RHICmdList);
 
diff --git a/Engine/Source/Runtime/Renderer/Private/DepthRendering.cpp b/Engine/Source/Runtime/Renderer/Private/DepthRendering.cpp
--- a/Engine/Source/Runtime/Renderer/Private/DepthRendering.cpp
+++ b/Engine/Source/Runtime/Renderer/Private/DepthRendering.cpp
@@ -1067,12 +1067,8 @@ bool FDeferredShadingSceneRenderer::RenderPrePass(FRHICommandListImmediate& RHIC
 {
 	bool bDepthWasCleared = false;
 
-#if WANTS_DRAW_MESH_EVENTS
-	static IConsoleVariable* CDBufferVar = IConsoleManager::Get().FindConsoleVariable(TEXT("r.DBuffer"));
-	bool bDBufferAllowed = CDBufferVar ? CDBufferVar->GetInt() != 0 : false;
-
-	SCOPED_DRAW_EVENTF(RHICmdList, PrePass, TEXT("PrePass %s %s"), GetDepthDrawingModeString(EarlyZPassMode), (bDBufferAllowed ? TEXT("(Forced by DBuffer)") : (bDitheredLODTransitionsUseStencil ? TEXT("(Forced by StencilLODDither)") : TEXT(""))));
-#endif
+	extern const TCHAR* GetDepthPassReason(bool bDitheredLODTransitionsUseStencil, ERHIFeatureLevel::Type FeatureLevel);
+	SCOPED_DRAW_EVENTF(RHICmdList, PrePass, TEXT("PrePass %s %s"), GetDepthDrawingModeString(EarlyZPassMode), GetDepthPassReason(bDitheredLODTransitionsUseStencil, FeatureLevel));
 
 	SCOPE_CYCLE_COUNTER(STAT_DepthDrawTime);
 
diff --git a/Engine/Source/Runtime/Renderer/Private/DistanceFieldShadowing.cpp b/Engine/Source/Runtime/Renderer/Private/DistanceFieldShadowing.cpp
--- a/Engine/Source/Runtime/Renderer/Private/DistanceFieldShadowing.cpp
+++ b/Engine/Source/Runtime/Renderer/Private/DistanceFieldShadowing.cpp
@@ -865,7 +865,7 @@ bool FDeferredShadingSceneRenderer::ShouldPrepareForDistanceFieldShadows() const
 		&& SupportsDistanceFieldShadows(Scene->GetFeatureLevel(), Scene->GetShaderPlatform());
 }
 
-void FProjectedShadowInfo::RenderRayTracedDistanceFieldProjection(FRHICommandListImmediate& RHICmdList, const FViewInfo& View) const
+void FProjectedShadowInfo::RenderRayTracedDistanceFieldProjection(FRHICommandListImmediate& RHICmdList, const FViewInfo& View, bool bProjectingForForwardShading) const
 {
 	if (SupportsDistanceFieldShadows(View.GetFeatureLevel(), View.GetShaderPlatform()))
 	{
@@ -980,18 +980,7 @@ void FProjectedShadowInfo::RenderRayTracedDistanceFieldProjection(FRHICommandLis
 				RHICmdList.SetRasterizerState(TStaticRasterizerState<FM_Solid, CM_None>::GetRHI());
 				RHICmdList.SetDepthStencilState(TStaticDepthStencilState<false, CF_Always>::GetRHI());
 				
-				if (bDirectionalLight)
-				{
-					// use R and G in Light Attenuation for directional lights
-					// CO_Min is needed to combine with far shadows which overlap the same depth range
-					RHICmdList.SetBlendState(TStaticBlendState<CW_RG, BO_Min, BF_One, BF_One>::GetRHI());
-				}
-				else
-				{
-					// use B and A in Light Attenuation
-					// CO_Min is needed to combine multiple shadow passes
-					RHICmdList.SetBlendState(TStaticBlendState<CW_BA, BO_Min, BF_One, BF_One, BO_Min, BF_One, BF_One>::GetRHI());
-				}
+				SetBlendStateForProjection(RHICmdList, bProjectingForForwardShading, false);
 
 				TShaderMapRef<FPostProcessVS> VertexShader(View.ShaderMap);
 
diff --git a/Engine/Source/Runtime/Renderer/Private/DistanceFieldSurfaceCacheLighting.cpp b/Engine/Source/Runtime/Renderer/Private/DistanceFieldSurfaceCacheLighting.cpp
--- a/Engine/Source/Runtime/Renderer/Private/DistanceFieldSurfaceCacheLighting.cpp
+++ b/Engine/Source/Runtime/Renderer/Private/DistanceFieldSurfaceCacheLighting.cpp
@@ -2743,7 +2743,7 @@ bool SupportsDistanceFieldAO(ERHIFeatureLevel::Type FeatureLevel, EShaderPlatfor
 		&& DoesPlatformSupportDistanceFieldAO(ShaderPlatform);
 }
 
-bool ShouldRenderDynamicSkyLight(const FScene* Scene, const FSceneViewFamily& ViewFamily)
+bool ShouldRenderDeferredDynamicSkyLight(const FScene* Scene, const FSceneViewFamily& ViewFamily)
 {
 	return Scene->SkyLight
 		&& Scene->SkyLight->ProcessedTexture
@@ -2751,14 +2751,14 @@ bool ShouldRenderDynamicSkyLight(const FScene* Scene, const FSceneViewFamily& Vi
 		&& !Scene->SkyLight->bHasStaticLighting
 		&& ViewFamily.EngineShowFlags.SkyLighting
 		&& Scene->GetFeatureLevel() >= ERHIFeatureLevel::SM4
-		&& !IsSimpleForwardShadingEnabled(Scene->GetShaderPlatform()) 
+		&& !IsAnyForwardShadingEnabled(Scene->GetShaderPlatform()) 
 		&& !ViewFamily.EngineShowFlags.VisualizeLightCulling;
 }
 
 bool FDeferredShadingSceneRenderer::ShouldPrepareForDistanceFieldAO() const
 {
 	return SupportsDistanceFieldAO(Scene->GetFeatureLevel(), Scene->GetShaderPlatform())
-		&& ((ShouldRenderDynamicSkyLight(Scene, ViewFamily) && Scene->SkyLight->bCastShadows && ViewFamily.EngineShowFlags.DistanceFieldAO)
+		&& ((ShouldRenderDeferredDynamicSkyLight(Scene, ViewFamily) && Scene->SkyLight->bCastShadows && ViewFamily.EngineShowFlags.DistanceFieldAO)
 			|| ViewFamily.EngineShowFlags.VisualizeMeshDistanceFields
 			|| ViewFamily.EngineShowFlags.VisualizeDistanceFieldAO
 			|| ViewFamily.EngineShowFlags.VisualizeDistanceFieldGI
@@ -3334,7 +3334,7 @@ bool FDeferredShadingSceneRenderer::ShouldRenderDistanceFieldAO() const
 
 void FDeferredShadingSceneRenderer::RenderDynamicSkyLighting(FRHICommandListImmediate& RHICmdList, const TRefCountPtr<IPooledRenderTarget>& VelocityTexture, TRefCountPtr<IPooledRenderTarget>& DynamicBentNormalAO)
 {
-	if (ShouldRenderDynamicSkyLight(Scene, ViewFamily))
+	if (ShouldRenderDeferredDynamicSkyLight(Scene, ViewFamily))
 	{
 		SCOPED_DRAW_EVENT(RHICmdList, SkyLightDiffuse);
 
diff --git a/Engine/Source/Runtime/Renderer/Private/DistanceFieldSurfaceCacheLighting.h b/Engine/Source/Runtime/Renderer/Private/DistanceFieldSurfaceCacheLighting.h
--- a/Engine/Source/Runtime/Renderer/Private/DistanceFieldSurfaceCacheLighting.h
+++ b/Engine/Source/Runtime/Renderer/Private/DistanceFieldSurfaceCacheLighting.h
@@ -649,4 +649,4 @@ private:
 
 extern void TrackGPUProgress(FRHICommandListImmediate& RHICmdList, uint32 DebugId);
 
-extern bool ShouldRenderDynamicSkyLight(const FScene* Scene, const FSceneViewFamily& ViewFamily);
+extern bool ShouldRenderDeferredDynamicSkyLight(const FScene* Scene, const FSceneViewFamily& ViewFamily);
diff --git a/Engine/Source/Runtime/Renderer/Private/HeightfieldLighting.cpp b/Engine/Source/Runtime/Renderer/Private/HeightfieldLighting.cpp
--- a/Engine/Source/Runtime/Renderer/Private/HeightfieldLighting.cpp
+++ b/Engine/Source/Runtime/Renderer/Private/HeightfieldLighting.cpp
@@ -1149,7 +1149,7 @@ void FHeightfieldLightingViewInfo::ComputeLighting(const FViewInfo& View, FRHICo
 
 		const FScene* Scene = (const FScene*)View.Family->Scene;
 
-		const float SkyLightIndirectScale = ShouldRenderDynamicSkyLight(Scene, *View.Family) ? Scene->SkyLight->IndirectLightingIntensity : 0;
+		const float SkyLightIndirectScale = ShouldRenderDeferredDynamicSkyLight(Scene, *View.Family) ? Scene->SkyLight->IndirectLightingIntensity : 0;
 
 		// Skip rendering if the DefaultLightFunctionMaterial isn't compiled yet
 		if (MaterialProxy->GetMaterial(FeatureLevel)->IsLightFunction())
