Copying //UE4/Dev-Rendering to //UE4/Dev-Main (Source: //UE4/Dev-Rendering @ 3054480)

diff --git a/Engine/Source/Runtime/Engine/Private/SceneUtils.cpp b/Engine/Source/Runtime/Engine/Private/SceneUtils.cpp
--- a/Engine/Source/Runtime/Engine/Private/SceneUtils.cpp
+++ b/Engine/Source/Runtime/Engine/Private/SceneUtils.cpp
@@ -4,6 +4,24 @@
 #include "RHI.h"
 #include "SceneUtils.h"
 
+
+#if HAS_GPU_STATS
+
+static TAutoConsoleVariable<int> CVarGPUStatsEnabled(
+	TEXT("r.GPUStatsEnabled"),
+	0,
+	TEXT("Enables or disables GPU stat recording"));
+
+// If this is enabled, the child stat timings will be included in their parents' times.
+// This presents problems for non-hierarchical stats if we're expecting them to add up
+// to the total GPU time, so we probably want this disabled
+#define GPU_STATS_CHILD_TIMES_INCLUDED 0
+
+DECLARE_FLOAT_COUNTER_STAT(TEXT("[TOTAL SCENE]"), Stat_GPU_Total, STATGROUP_GPU);
+
+#endif //HAS_GPU_STATS
+
+
 #if WANTS_DRAW_MESH_EVENTS
 
 template<typename TRHICmdList>
@@ -83,3 +101,365 @@ bool IsMobileHDRMosaic()
 			return !(GSupportsHDR32bppEncodeModeIntrinsic && GSupportsShaderFramebufferFetch);
 	}
 }
+
+#if HAS_GPU_STATS
+static const int32 NumGPUProfilerBufferedFrames = 4;
+
+/*-----------------------------------------------------------------------------
+FRealTimeGPUProfilerEvent class
+-----------------------------------------------------------------------------*/
+class FRealtimeGPUProfilerEvent
+{
+public:
+	static const uint64 InvalidQueryResult = 0xFFFFFFFFFFFFFFFFull;
+
+public:
+	FRealtimeGPUProfilerEvent(const TStatId& InStatId)
+		: StartResultMicroseconds(InvalidQueryResult)
+		, EndResultMicroseconds(InvalidQueryResult)
+		, FrameNumber(-1)
+		, bInsideQuery(false)
+	{
+		StatName = InStatId.GetName();
+	}
+
+	void Begin(FRHICommandListImmediate& RHICmdList)
+	{
+		check(IsInRenderingThread());
+		check(!bInsideQuery);
+		bInsideQuery = true;
+
+		if (!StartQuery)
+		{
+			StartQuery = RHICmdList.CreateRenderQuery(RQT_AbsoluteTime);
+			EndQuery = RHICmdList.CreateRenderQuery(RQT_AbsoluteTime);
+		}
+
+		RHICmdList.EndRenderQuery(StartQuery);
+		StartResultMicroseconds = InvalidQueryResult;
+		EndResultMicroseconds = InvalidQueryResult;
+		FrameNumber = GFrameNumberRenderThread;
+	}
+
+	void End(FRHICommandListImmediate& RHICmdList)
+	{
+		check(IsInRenderingThread());
+		check(bInsideQuery);
+		bInsideQuery = false;
+		RHICmdList.EndRenderQuery(EndQuery);
+	}
+
+	bool GatherQueryResults(FRHICommandListImmediate& RHICmdList)
+	{
+		// Get the query results which are still outstanding
+		if (StartResultMicroseconds == InvalidQueryResult)
+		{
+			if (!RHICmdList.GetRenderQueryResult(StartQuery, StartResultMicroseconds, false))
+			{
+				StartResultMicroseconds = InvalidQueryResult;
+			}
+		}
+		if (EndResultMicroseconds == InvalidQueryResult)
+		{
+			if (!RHICmdList.GetRenderQueryResult(EndQuery, EndResultMicroseconds, false))
+			{
+				EndResultMicroseconds = InvalidQueryResult;
+			}
+		}
+		return HasValidResult();
+	}
+
+	float GetResultMS() const
+	{
+		check(HasValidResult());
+		if (EndResultMicroseconds < StartResultMicroseconds)
+		{
+			// This should never happen...
+			return 0.0f;
+		}
+		return float(EndResultMicroseconds - StartResultMicroseconds) / 1000.0f;
+	}
+
+	bool HasValidResult() const
+	{
+		return StartResultMicroseconds != FRealtimeGPUProfilerEvent::InvalidQueryResult && EndResultMicroseconds != FRealtimeGPUProfilerEvent::InvalidQueryResult;
+	}
+
+	const FName& GetStatName() const
+	{
+		return StatName;
+	}
+
+private:
+	FRenderQueryRHIRef StartQuery;
+	FRenderQueryRHIRef EndQuery;
+	FName StatName;
+	uint64 StartResultMicroseconds;
+	uint64 EndResultMicroseconds;
+	uint32 FrameNumber;
+
+	bool bInsideQuery;
+};
+
+/*-----------------------------------------------------------------------------
+FRealtimeGPUProfilerFrame class
+Container for a single frame's GPU stats
+-----------------------------------------------------------------------------*/
+class FRealtimeGPUProfilerFrame
+{
+public:
+	FRealtimeGPUProfilerFrame()
+		: FrameNumber(-1)
+	{}
+
+	~FRealtimeGPUProfilerFrame()
+	{
+		Clear();
+	}
+
+	void PushEvent(FRHICommandListImmediate& RHICmdList, TStatId StatId)
+	{
+#if GPU_STATS_CHILD_TIMES_INCLUDED == 0
+		if (EventStack.Num() > 0)
+		{
+			// GPU Stats are not hierarchical. If we already have an event in the stack, 
+			// we need end it and resume it once the child event completes 
+			FRealtimeGPUProfilerEvent* ParentEvent = EventStack.Last();
+			ParentEvent->End(RHICmdList);
+		}
+#endif
+		FRealtimeGPUProfilerEvent* Event = CreateNewEvent(StatId);
+		EventStack.Push(Event);
+		StatStack.Push(StatId);
+		Event->Begin(RHICmdList);
+	}
+
+	void PopEvent(FRHICommandListImmediate& RHICmdList)
+	{
+		FRealtimeGPUProfilerEvent* Event = EventStack.Pop();
+		StatStack.Pop();
+		Event->End(RHICmdList);
+
+#if GPU_STATS_CHILD_TIMES_INCLUDED == 0
+		if (EventStack.Num() > 0)
+		{
+			// Resume the parent event (requires creation of a new FRealtimeGPUProfilerEvent)
+			TStatId PrevStatId = StatStack.Last();
+			FRealtimeGPUProfilerEvent* ResumedEvent = CreateNewEvent(PrevStatId);
+			EventStack.Last() = ResumedEvent;
+			ResumedEvent->Begin(RHICmdList);
+		}
+#endif
+
+	}
+
+	void Clear()
+	{
+		EventStack.Empty();
+		StatStack.Empty();
+
+		for (int Index = 0; Index < GpuProfilerEvents.Num(); Index++)
+		{
+			if (GpuProfilerEvents[Index])
+			{
+				delete GpuProfilerEvents[Index];
+			}
+		}
+		GpuProfilerEvents.Empty();
+	}
+
+	bool UpdateStats(FRHICommandListImmediate& RHICmdList)
+	{
+		// Gather any remaining results and check all the results are ready
+		for (int Index = 0; Index < GpuProfilerEvents.Num(); Index++)
+		{
+			FRealtimeGPUProfilerEvent* Event = GpuProfilerEvents[Index];
+			check(Event != nullptr);
+			if (!Event->HasValidResult())
+			{
+				Event->GatherQueryResults(RHICmdList);
+			}
+			if (!Event->HasValidResult())
+			{
+				// The frame isn't ready yet. Don't update stats - we'll try again next frame. 
+				return false;
+			}
+		}
+
+		float TotalMS = 0.0f;
+		// Update the stats
+		TMap<FName, bool> StatSeenMap;
+		for (int Index = 0; Index < GpuProfilerEvents.Num(); Index++)
+		{
+			FRealtimeGPUProfilerEvent* Event = GpuProfilerEvents[Index];
+			check(Event != nullptr);
+			check(Event->HasValidResult());
+			EStatOperation::Type StatOp;
+			const FName& StatName = Event->GetStatName();
+
+			// Check if we've seen this stat yet 
+			if (StatSeenMap.Find(StatName) == nullptr)
+			{
+				StatSeenMap.Add(StatName, true);
+				StatOp = EStatOperation::Set;
+			}
+			else
+			{
+				// Stat was seen before, so accumulate 
+				StatOp = EStatOperation::Add;
+			}
+			float ResultMS = Event->GetResultMS();
+			FThreadStats::AddMessage(StatName, StatOp, double(ResultMS));
+			TotalMS += ResultMS;
+		}
+
+		FThreadStats::AddMessage( GET_STATFNAME(Stat_GPU_Total), EStatOperation::Set, double(TotalMS) );
+		return true;
+	}
+
+private:
+	FRealtimeGPUProfilerEvent* CreateNewEvent(const TStatId& StatId)
+	{
+		FRealtimeGPUProfilerEvent* NewEvent = new FRealtimeGPUProfilerEvent(StatId);
+		GpuProfilerEvents.Add(NewEvent);
+		return NewEvent;
+	}
+
+	TArray<FRealtimeGPUProfilerEvent*> GpuProfilerEvents;
+	TArray<FRealtimeGPUProfilerEvent*> EventStack;
+	TArray<TStatId> StatStack;
+	uint32 FrameNumber;
+};
+
+/*-----------------------------------------------------------------------------
+FRealtimeGPUProfiler
+-----------------------------------------------------------------------------*/
+FRealtimeGPUProfiler* FRealtimeGPUProfiler::Instance = nullptr;
+
+FRealtimeGPUProfiler* FRealtimeGPUProfiler::Get()
+{
+	if (Instance == nullptr)
+	{
+		Instance = new FRealtimeGPUProfiler;
+	}
+	return Instance;
+}
+
+FRealtimeGPUProfiler::FRealtimeGPUProfiler()
+	: WriteBufferIndex(-1)
+	, ReadBufferIndex(-1)
+	, WriteFrameNumber(-1)
+{
+	for (int Index = 0; Index < NumGPUProfilerBufferedFrames; Index++)
+	{
+		Frames.Add(new FRealtimeGPUProfilerFrame());
+	}
+}
+
+void FRealtimeGPUProfiler::Release()
+{
+	for (int Index = 0; Index < Frames.Num(); Index++)
+	{
+		delete Frames[Index];
+	}
+	Frames.Empty();
+}
+
+void FRealtimeGPUProfiler::Update(FRHICommandListImmediate& RHICmdList)
+{
+	check(Frames.Num() > 0);
+	if (WriteFrameNumber == GFrameNumberRenderThread)
+	{
+		// Multiple views: only update if this is a new frame. Otherwise we just concatenate the stats
+		return;
+	}
+	WriteFrameNumber = GFrameNumberRenderThread;
+	WriteBufferIndex = (WriteBufferIndex + 1) % Frames.Num();
+	Frames[WriteBufferIndex]->Clear();
+
+	// If the write buffer catches the read buffer, we need to advance the read buffer before we write over it 
+	if (WriteBufferIndex == ReadBufferIndex)
+	{
+		ReadBufferIndex = (ReadBufferIndex - 1) % Frames.Num();
+	}
+	UpdateStats(RHICmdList);
+}
+
+void FRealtimeGPUProfiler::PushEvent(FRHICommandListImmediate& RHICmdList, TStatId StatId)
+{
+	check(Frames.Num() > 0);
+	if (WriteBufferIndex >= 0)
+	{
+		Frames[WriteBufferIndex]->PushEvent(RHICmdList, StatId);
+	}
+}
+
+void FRealtimeGPUProfiler::PopEvent(FRHICommandListImmediate& RHICmdList)
+{
+	check(Frames.Num() > 0);
+	if (WriteBufferIndex >= 0)
+	{
+		Frames[WriteBufferIndex]->PopEvent(RHICmdList);
+	}
+}
+
+void FRealtimeGPUProfiler::UpdateStats(FRHICommandListImmediate& RHICmdList)
+{
+	check(IsInRenderingThread());
+	if (GSupportsTimestampRenderQueries == false || !CVarGPUStatsEnabled.GetValueOnRenderThread() )
+	{
+		return;
+	}
+	check(Frames.Num() > 0);
+	bool bAdvanceReadIndex = false;
+	if (ReadBufferIndex == -1)
+	{
+		bAdvanceReadIndex = true;
+	}
+	else
+	{
+		// If the read frame is valid, update the stats
+		if (Frames[ReadBufferIndex]->UpdateStats(RHICmdList))
+		{
+			// On a successful read, advance the ReadBufferIndex
+			bAdvanceReadIndex = true;
+		}
+	}
+	if (bAdvanceReadIndex)
+	{
+		ReadBufferIndex = (ReadBufferIndex + 1) % Frames.Num();
+	}
+}
+
+/*-----------------------------------------------------------------------------
+FScopedGPUStatEvent
+-----------------------------------------------------------------------------*/
+void FScopedGPUStatEvent::Begin(FRHICommandList& InRHICmdList, TStatId InStatID)
+{
+	check(IsInRenderingThread());
+	if (GSupportsTimestampRenderQueries == false || !CVarGPUStatsEnabled.GetValueOnRenderThread())
+	{
+		return;
+	}
+
+	// Non-immediate commandlists are not supported (silently fail)
+	if (InRHICmdList.IsImmediate())
+	{
+		RHICmdList = (FRHICommandListImmediate*)&InRHICmdList;
+		FRealtimeGPUProfiler::Get()->PushEvent(*RHICmdList, InStatID);
+	}
+}
+
+void FScopedGPUStatEvent::End()
+{
+	check(IsInRenderingThread());
+	if (GSupportsTimestampRenderQueries == false || !CVarGPUStatsEnabled.GetValueOnRenderThread() )
+	{
+		return;
+	}
+	if (RHICmdList != nullptr)
+	{
+		FRealtimeGPUProfiler::Get()->PopEvent(*RHICmdList);
+	}
+}
+#endif // HAS_GPU_STATS
diff --git a/Engine/Source/Runtime/Engine/Private/ShaderCompiler/ShaderCompiler.cpp b/Engine/Source/Runtime/Engine/Private/ShaderCompiler/ShaderCompiler.cpp
--- a/Engine/Source/Runtime/Engine/Private/ShaderCompiler/ShaderCompiler.cpp
+++ b/Engine/Source/Runtime/Engine/Private/ShaderCompiler/ShaderCompiler.cpp
@@ -2395,8 +2395,8 @@ void GlobalBeginCompileShader(
 		// List mostly comes from set of characters not allowed by windows in a path.  Just try to rename a file and type one of these for the list.
 		Input.DumpDebugInfoPath.ReplaceInline(TEXT("<"), TEXT("("));
 		Input.DumpDebugInfoPath.ReplaceInline(TEXT(">"), TEXT(")"));
-		Input.DumpDebugInfoPath.ReplaceInline(TEXT("::"), TEXT("=="));		
-		Input.DumpDebugInfoPath.ReplaceInline(TEXT("|"), TEXT("_"));		
+		Input.DumpDebugInfoPath.ReplaceInline(TEXT("::"), TEXT("=="));
+		Input.DumpDebugInfoPath.ReplaceInline(TEXT("|"), TEXT("_"));
 		Input.DumpDebugInfoPath.ReplaceInline(TEXT("*"), TEXT("-"));
 		Input.DumpDebugInfoPath.ReplaceInline(TEXT("?"), TEXT("!"));
 		Input.DumpDebugInfoPath.ReplaceInline(TEXT("\""), TEXT("\'"));
@@ -2409,12 +2409,12 @@ void GlobalBeginCompileShader(
 
 	// Add the appropriate definitions for the shader frequency.
 	{
-		Input.Environment.SetDefine(TEXT("PIXELSHADER"),	Target.Frequency == SF_Pixel);
-		Input.Environment.SetDefine(TEXT("DOMAINSHADER"),	Target.Frequency == SF_Domain);
-		Input.Environment.SetDefine(TEXT("HULLSHADER"),		Target.Frequency == SF_Hull);
-		Input.Environment.SetDefine(TEXT("VERTEXSHADER"),	Target.Frequency == SF_Vertex);
-		Input.Environment.SetDefine(TEXT("GEOMETRYSHADER"),	Target.Frequency == SF_Geometry);
-		Input.Environment.SetDefine(TEXT("COMPUTESHADER"),	Target.Frequency == SF_Compute);
+		Input.Environment.SetDefine(TEXT("PIXELSHADER"), Target.Frequency == SF_Pixel);
+		Input.Environment.SetDefine(TEXT("DOMAINSHADER"), Target.Frequency == SF_Domain);
+		Input.Environment.SetDefine(TEXT("HULLSHADER"), Target.Frequency == SF_Hull);
+		Input.Environment.SetDefine(TEXT("VERTEXSHADER"), Target.Frequency == SF_Vertex);
+		Input.Environment.SetDefine(TEXT("GEOMETRYSHADER"), Target.Frequency == SF_Geometry);
+		Input.Environment.SetDefine(TEXT("COMPUTESHADER"), Target.Frequency == SF_Compute);
 	}
 
 	// Set instanced stereo define
@@ -2454,7 +2454,7 @@ void GlobalBeginCompileShader(
 			Input.Environment.CompilerFlags.Add(CFLAG_Debug);
 		}
 	}
-	
+
 	{
 		static const auto CVar = IConsoleManager::Get().FindConsoleVariable(TEXT("r.Shaders.KeepDebugInfo"));
 
@@ -2487,7 +2487,7 @@ void GlobalBeginCompileShader(
 	{
 		FString ShaderPDBRoot;
 		GConfig->GetString(TEXT("DevOptions.Shaders"), TEXT("ShaderPDBRoot"), ShaderPDBRoot, GEngineIni);
-		if ( !ShaderPDBRoot.IsEmpty() )
+		if (!ShaderPDBRoot.IsEmpty())
 		{
 			Input.Environment.SetDefine(TEXT("SHADER_PDB_ROOT"), *ShaderPDBRoot);
 		}
@@ -2503,7 +2503,7 @@ void GlobalBeginCompileShader(
 		Input.Environment.SetDefine(TEXT("DXT5_NORMALMAPS"), CVar ? (CVar->GetValueOnGameThread() != 0) : 0);
 	}
 
-	if(bAllowDevelopmentShaderCompile)
+	if (bAllowDevelopmentShaderCompile)
 	{
 		static const auto CVar = IConsoleManager::Get().FindTConsoleVariableDataInt(TEXT("r.CompileShadersForDevelopment"));
 		Input.Environment.SetDefine(TEXT("COMPILE_SHADERS_FOR_DEVELOPMENT"), CVar ? (CVar->GetValueOnGameThread() != 0) : 0);
@@ -2539,6 +2539,15 @@ void GlobalBeginCompileShader(
 		Input.Environment.SetDefine(TEXT("FORWARD_SHADING"), CVar ? (CVar->GetInt() != 0) : 0);
 	}
 
+	if (GSupportsRenderTargetWriteMask)
+	{
+		Input.Environment.SetDefine(TEXT("PLATFORM_SUPPORTS_RENDERTARGET_WRITE_MASK"), 1);
+	}
+	else
+	{
+		Input.Environment.SetDefine(TEXT("PLATFORM_SUPPORTS_RENDERTARGET_WRITE_MASK"), 0);
+	}
+
 	NewJobs.Add(NewJob);
 }
 
diff --git a/Engine/Source/Runtime/Engine/Private/SkeletalRenderCPUSkin.cpp b/Engine/Source/Runtime/Engine/Private/SkeletalRenderCPUSkin.cpp
--- a/Engine/Source/Runtime/Engine/Private/SkeletalRenderCPUSkin.cpp
+++ b/Engine/Source/Runtime/Engine/Private/SkeletalRenderCPUSkin.cpp
@@ -23,6 +23,7 @@
 #include "EnginePrivate.h"
 #include "SkeletalRenderCPUSkin.h"
 #include "Animation/MorphTarget.h"
+#include "GPUSkinVertexFactory.h"
 
 struct FMorphTargetDelta;
 
@@ -807,11 +808,11 @@ static void SkinVertices(FFinalSkinVertex* DestVertex, FMatrix* ReferenceToLocal
 	TArray<FMorphTargetInfo> MorphEvalInfos;
 	uint32 NumValidMorphs = InitEvalInfos(ActiveMorphTargets, MorphTargetWeights, LODIndex, MorphEvalInfos);
 
-	static const auto MaxBonesVar = IConsoleManager::Get().FindTConsoleVariableDataInt(TEXT("Compat.MAX_GPUSKIN_BONES"));
-	const int32 MaxGPUSkinBones = MaxBonesVar->GetValueOnAnyThread();
+	const uint32 MaxGPUSkinBones = FGPUBaseSkinVertexFactory::GetMaxGPUSkinBones();
+	check(MaxGPUSkinBones <= FGPUBaseSkinVertexFactory::GHardwareMaxGPUSkinBones);
 
 	// Prefetch all matrices
-	for ( int32 MatrixIndex=0; MatrixIndex < MaxGPUSkinBones; MatrixIndex+=2 )
+	for ( uint32 MatrixIndex=0; MatrixIndex < MaxGPUSkinBones; MatrixIndex+=2 )
 	{
 		FPlatformMisc::Prefetch( ReferenceToLocal + MatrixIndex );
 	}
diff --git a/Engine/Source/Runtime/Engine/Private/Streaming/AsyncTextureStreaming.cpp b/Engine/Source/Runtime/Engine/Private/Streaming/AsyncTextureStreaming.cpp
--- a/Engine/Source/Runtime/Engine/Private/Streaming/AsyncTextureStreaming.cpp
+++ b/Engine/Source/Runtime/Engine/Private/Streaming/AsyncTextureStreaming.cpp
@@ -210,14 +210,11 @@ void FAsyncTextureStreamingTask::UpdateBudgetedMips_Async(int64& MemoryUsed, int
 		// Sort texture, having those that should be dropped first.
 		PrioritizedTextures.Sort(FCompareTextureByRetentionPriority(StreamingTextures));
 
-		bool bBudgetIsChanging = true;
-		int32 NumDroppedMips = 0;
-		while (MemoryBudgeted > MemoryBudget && bBudgetIsChanging && !IsAborted())
+		// When using mip bias per texture, we first reduce the maximum resolutions (if used) in order to fit.
+		for (int32 NumDroppedMips = 0; NumDroppedMips < MaxPerTextureMipBias && MemoryBudgeted > MemoryBudget && !IsAborted(); ++NumDroppedMips)
 		{
-			bBudgetIsChanging = false;
-			++NumDroppedMips;
+			bool bBudgetCanChange = false;
 
-			// Drop from the lowest priority first (starting with last elements)
 			for (int32 PriorityIndex = PrioritizedTextures.Num() - 1; PriorityIndex >= 0 && MemoryBudgeted > MemoryBudget && !IsAborted(); --PriorityIndex)
 			{
 				int32 TextureIndex = PrioritizedTextures[PriorityIndex];
@@ -225,24 +222,71 @@ void FAsyncTextureStreamingTask::UpdateBudgetedMips_Async(int64& MemoryUsed, int
 
 				FStreamingTexture& StreamingTexture = StreamingTextures[TextureIndex];
 
-				// If this texture has already missing mips for its normal quality, don't drop more than required..
-				if (StreamingTexture.NumMissingMips >= NumDroppedMips)
+				if (StreamingTexture.BudgetedMips <= StreamingTexture.MinAllowedMips)
 				{
-					bBudgetIsChanging = true; // Needed to prevent aborting loop.
+					// Don't try this one again.
+					PrioritizedTextures[PriorityIndex] = INDEX_NONE;
 					continue;
 				}
+				else
+				{
+					bBudgetCanChange = true;
+				}
 
-				int64 GivenMemory = StreamingTexture.DropOneMip_Async(MaxPerTextureMipBias);
+				if (!StreamingTexture.IsMaxResolutionAffectedByGlobalBias()) continue;
 
-				if (GivenMemory > 0)
+				// If the texture requires a high resolution mip, consider dropping it. 
+				// When considering dropping the first mip, only textures using the first mip will drop their resolution, 
+				// But when considering dropping the second mip, textures using their first and second mips will loose it.
+				if (StreamingTexture.MaxAllowedMips + StreamingTexture.BudgetMipBias - NumDroppedMips <= StreamingTexture.BudgetedMips)
 				{
-					MemoryBudgeted -= GivenMemory;
-					bBudgetIsChanging = true;
+					MemoryBudgeted -= StreamingTexture.DropMaxResolution_Async(NumDroppedMips + 1 - StreamingTexture.BudgetMipBias);
 				}
-				else // Nothing to drop
+			}
+
+			if (!bBudgetCanChange)
+			{
+				break;
+			}
+		}
+
+		// Then reduce all resolution.
+		while (MemoryBudgeted > MemoryBudget && !IsAborted())
+		{
+			bool bBudgetCanChange = false;
+
+			// Drop from the lowest priority first (starting with last elements)
+			for (int32 PriorityIndex = PrioritizedTextures.Num() - 1; PriorityIndex >= 0 && MemoryBudgeted > MemoryBudget && !IsAborted(); --PriorityIndex)
+			{
+				int32 TextureIndex = PrioritizedTextures[PriorityIndex];
+				if (TextureIndex == INDEX_NONE) continue;
+
+				FStreamingTexture& StreamingTexture = StreamingTextures[TextureIndex];
+
+				if (StreamingTexture.BudgetedMips <= StreamingTexture.MinAllowedMips)
 				{
-					PrioritizedTextures[PriorityIndex] = INDEX_NONE; // Don't try this one again.
+					// Don't try this one again.
+					PrioritizedTextures[PriorityIndex] = INDEX_NONE;
+					continue;
+				}
+				else
+				{
+					bBudgetCanChange = true;
+				}
+
+				// If this texture has already missing mips for its normal quality, don't drop more than required..
+				if (StreamingTexture.NumMissingMips > 0)
+				{
+					--StreamingTexture.NumMissingMips;
+					continue;
 				}
+
+				MemoryBudgeted -= StreamingTexture.DropOneMip_Async(MaxPerTextureMipBias);
+			}
+
+			if (!bBudgetCanChange)
+			{
+				break;
 			}
 		}
 	}
@@ -432,6 +476,7 @@ void FAsyncTextureStreamingTask::UpdateStats_Async()
 #if STATS
 	FTextureStreamingStats& Stats = StreamingManager.GatheredStats;
 	FTextureStreamingStats& PrevStats = StreamingManager.DisplayedStats;
+	FTextureStreamingSettings& Settings = StreamingManager.Settings;
 	TArray<FStreamingTexture>& StreamingTextures = StreamingManager.StreamingTextures;
 
 	Stats.TexturePool = PoolSize;
diff --git a/Engine/Source/Runtime/Engine/Private/Streaming/StreamingManagerTexture.cpp b/Engine/Source/Runtime/Engine/Private/Streaming/StreamingManagerTexture.cpp
--- a/Engine/Source/Runtime/Engine/Private/Streaming/StreamingManagerTexture.cpp
+++ b/Engine/Source/Runtime/Engine/Private/Streaming/StreamingManagerTexture.cpp
@@ -493,7 +493,10 @@ void FStreamingManagerTexture::AddStreamingTexture( UTexture2D* Texture )
 {
 	// Adds the new texture to the Pending list, to avoid reallocation of the thread-safe StreamingTextures array.
 	check(Texture->StreamingIndex == INDEX_NONE);
-	Texture->StreamingIndex = PendingStreamingTextures.Add(Texture);;
+	Texture->StreamingIndex = PendingStreamingTextures.Add(Texture);
+
+	// Notify that this texture ptr is valid.
+	ReferencedTextures.Add(Texture);
 }
 
 /**
@@ -515,6 +518,9 @@ void FStreamingManagerTexture::RemoveStreamingTexture( UTexture2D* Texture )
 	}
 
 	Texture->StreamingIndex = INDEX_NONE;
+
+	// Remove reference to this texture.
+	ReferencedTextures.Remove(Texture);
 }
 
 /** Called when an actor is spawned. */
@@ -603,7 +609,12 @@ void FStreamingManagerTexture::SetTexturesRemovedTimestamp(const FRemovedTexture
 	const double CurrentTime = FApp::GetCurrentTime();
 	for ( int32 TextureIndex=0; TextureIndex < RemovedTextures.Num(); ++TextureIndex )
 	{
-		FStreamingTexture* StreamingTexture = GetStreamingTexture( RemovedTextures[TextureIndex] );
+		// When clearing references to textures, those textures could be already deleted.
+		// This happens because we don't clear texture references in RemoveStreamingTexture.
+		const UTexture2D* Texture = RemovedTextures[TextureIndex];
+		if (!ReferencedTextures.Contains(Texture)) continue;
+
+		FStreamingTexture* StreamingTexture = GetStreamingTexture(Texture);
 		if ( StreamingTexture )
 		{
 			StreamingTexture->InstanceRemovedTimestamp = CurrentTime;
diff --git a/Engine/Source/Runtime/Engine/Private/Streaming/StreamingManagerTexture.h b/Engine/Source/Runtime/Engine/Private/Streaming/StreamingManagerTexture.h
--- a/Engine/Source/Runtime/Engine/Private/Streaming/StreamingManagerTexture.h
+++ b/Engine/Source/Runtime/Engine/Private/Streaming/StreamingManagerTexture.h
@@ -218,6 +218,9 @@ protected:
 		/** All streaming UTexture2D objects. */
 		TArray<FStreamingTexture> StreamingTextures;
 
+		/** All the textures referenced in StreamingTextures. Used to handled deleted textures.  */
+		TSet<const UTexture2D*> ReferencedTextures;
+
 		/** Index of the StreamingTexture that will be updated next by UpdateStreamingTextures(). */
 		int32 CurrentUpdateStreamingTextureIndex;
 //END: Thread-safe functions and data
diff --git a/Engine/Source/Runtime/Engine/Private/Streaming/StreamingTexture.cpp b/Engine/Source/Runtime/Engine/Private/Streaming/StreamingTexture.cpp
--- a/Engine/Source/Runtime/Engine/Private/Streaming/StreamingTexture.cpp
+++ b/Engine/Source/Runtime/Engine/Private/Streaming/StreamingTexture.cpp
@@ -89,7 +89,7 @@ void FStreamingTexture::UpdateDynamicData(const int32 NumStreamedMips[TEXTUREGRO
 			LODBias = FMath::Max<int32>(Texture->GetCachedLODBias() - NumCinematicMipLevels, 0);
 
 			// Reduce the max allowed resolution according to LODBias if the texture group allows it.
-			if (LODGroup != TEXTUREGROUP_HierarchicalLOD && !bIsTerrainTexture)
+			if (IsMaxResolutionAffectedByGlobalBias())
 			{
 				LODBias += Settings.GlobalMipBiasAsInt();
 			}
@@ -166,8 +166,7 @@ void FStreamingTexture::UpdateStreamingStatus()
 float FStreamingTexture::GetExtraBoost(TextureGroup	LODGroup, const FTextureStreamingSettings& Settings)
 {
 	// When accurate distance computation, we need to relax the distance otherwhise it get's too conservative. (ex 513 goes to 1024)
-	const bool bUseApproxDistance = CVarStreamingUseNewMetrics.GetValueOnAnyThread() == 0;
-	const float DistanceScale = bUseApproxDistance ? 1.f : .71f;
+	const float DistanceScale = Settings.bUseNewMetrics ? .71f : 1.f;
 
 	if (LODGroup == TEXTUREGROUP_Terrain_Heightmap || LODGroup == TEXTUREGROUP_Terrain_Weightmap) 
 	{
@@ -203,18 +202,17 @@ void FStreamingTexture::SetPerfectWantedMips_Async(float MaxSize, float MaxSize_
 	VisibleWantedMips = GetWantedMipsFromSize(MaxSize_VisibleOnly);
 	bLooksLowRes = InLooksLowRes; // Things like lightmaps, HLOD and close instances.
 
-	// Here BudgetMipBias is added to NumMissingMips to prevent the same texture from dropping resolution too early (as it is already included in MaxMip)
-
 	// Terrain, Forced Fully Load and Things that already look bad are not affected by hidden scale.
 	if (bIsTerrainTexture || bForceFullyLoadHeuristic || bLooksLowRes)
 	{
 		HiddenWantedMips = GetWantedMipsFromSize(MaxSize);
-		NumMissingMips = BudgetMipBias; // No impact for terrains as there are not allowed to drop mips.
+		NumMissingMips = 0; // No impact for terrains as there are not allowed to drop mips.
 	}
 	else
 	{
 		HiddenWantedMips = GetWantedMipsFromSize(MaxSize * Settings.HiddenPrimitiveScale);
-		NumMissingMips = BudgetMipBias + FMath::Max<int32>(GetWantedMipsFromSize(MaxSize) - FMath::Max<int32>(VisibleWantedMips, HiddenWantedMips), 0);
+		// NumMissingMips contains the number of mips not loaded because of HiddenPrimitiveScale. When out of budget, those texture will be considered as already sacrificed.
+		NumMissingMips = FMath::Max<int32>(GetWantedMipsFromSize(MaxSize) - FMath::Max<int32>(VisibleWantedMips, HiddenWantedMips), 0);
 	}
 }
 
@@ -250,22 +248,38 @@ int64 FStreamingTexture::UpdateRetentionPriority_Async()
 	}
 }
 
-int64 FStreamingTexture::DropOneMip_Async(int32 MaxPerTextureMipBias)
+int64 FStreamingTexture::DropMaxResolution_Async(int32 NumDroppedMips)
 {
-	if (Texture && BudgetedMips > MinAllowedMips)
+	if (Texture)
 	{
-		// Every time we are required to DropOneMip from max resolution, this increase the BudgetMipBias.
-		if (BudgetedMips == MaxAllowedMips && VisibleWantedMips == MaxAllowedMips && BudgetMipBias < MaxPerTextureMipBias)
+		// Don't drop bellow min allowed mips.
+		NumDroppedMips = FMath::Min<int32>(MaxAllowedMips - MinAllowedMips, NumDroppedMips);
+
+		if (NumDroppedMips > 0)
 		{
 			// Decrease MaxAllowedMips and increase BudgetMipBias (as it should include it)
-			--MaxAllowedMips;
-			++BudgetMipBias;
+			MaxAllowedMips -= NumDroppedMips;
+			BudgetMipBias += NumDroppedMips;
+
+			if (BudgetedMips > MaxAllowedMips)
+			{
+				const int64 FreedMemory = GetSize(BudgetedMips) - GetSize(MaxAllowedMips);
+
+				BudgetedMips = MaxAllowedMips;
+				VisibleWantedMips = FMath::Min<int32>(VisibleWantedMips, MaxAllowedMips);
+				HiddenWantedMips = FMath::Min<int32>(HiddenWantedMips, MaxAllowedMips);
 
-			// Update the wanted mips as if they were computed with the new MaxAllowedMips. (for stats)
-			VisibleWantedMips = FMath::Min<int32>(VisibleWantedMips, MaxAllowedMips);
-			HiddenWantedMips = FMath::Min<int32>(HiddenWantedMips, MaxAllowedMips);
+				return FreedMemory;
+			}
 		}
+	}
+	return 0;
+}
 
+int64 FStreamingTexture::DropOneMip_Async(int32 MaxPerTextureMipBias)
+{
+	if (Texture && BudgetedMips > MinAllowedMips)
+	{
 		--BudgetedMips;
 		return GetSize(BudgetedMips + 1) - GetSize(BudgetedMips);
 	}
diff --git a/Engine/Source/Runtime/Engine/Private/Streaming/StreamingTexture.h b/Engine/Source/Runtime/Engine/Private/Streaming/StreamingTexture.h
--- a/Engine/Source/Runtime/Engine/Private/Streaming/StreamingTexture.h
+++ b/Engine/Source/Runtime/Engine/Private/Streaming/StreamingTexture.h
@@ -48,6 +48,9 @@ struct FStreamingTexture
 	/** Init BudgetedMip and update RetentionPriority. Returns the size that would be taken if all budgeted mips where loaded. */
 	int64 UpdateRetentionPriority_Async();
 
+	/** Reduce the maximum allowed resolution by 1 mip. Return the size freed by doing so. */
+	int64 DropMaxResolution_Async(int32 NumDroppedMips);
+
 	/** Reduce BudgetedMip by 1 and return the size freed by doing so. */
 	int64 DropOneMip_Async(int32 MaxPerTextureMipBias = 0);
 
@@ -64,6 +67,9 @@ struct FStreamingTexture
 
 	FORCEINLINE int32 GetPerfectWantedMips() const { return FMath::Max<int32>(VisibleWantedMips,  HiddenWantedMips); }
 
+	// Whether this texture can be affected by Global Bias and Budget Bias per texture.
+	FORCEINLINE bool IsMaxResolutionAffectedByGlobalBias() const { return LODGroup != TEXTUREGROUP_HierarchicalLOD && !bIsTerrainTexture; }
+
 	/***************************************************************
 	 * Members initialized when this is constructed => NEVER CHANGES
 	 ***************************************************************/
diff --git a/Engine/Source/Runtime/Engine/Private/UserInterface/Canvas.cpp b/Engine/Source/Runtime/Engine/Private/UserInterface/Canvas.cpp
--- a/Engine/Source/Runtime/Engine/Private/UserInterface/Canvas.cpp
+++ b/Engine/Source/Runtime/Engine/Private/UserInterface/Canvas.cpp
@@ -1372,6 +1372,8 @@ ESimpleElementBlendMode FCanvas::BlendToSimpleElementBlend(EBlendMode BlendMode)
 			return SE_BLEND_Additive;
 		case BLEND_Modulate:
 			return SE_BLEND_Modulate;
+		case BLEND_AlphaComposite:
+			return SE_BLEND_AlphaComposite;
 		case BLEND_Translucent:
 		default:
 			return SE_BLEND_Translucent;
diff --git a/Engine/Source/Runtime/Engine/Public/PixelFormat.h b/Engine/Source/Runtime/Engine/Public/PixelFormat.h
--- a/Engine/Source/Runtime/Engine/Public/PixelFormat.h
+++ b/Engine/Source/Runtime/Engine/Public/PixelFormat.h
@@ -62,7 +62,8 @@ enum EPixelFormat
 	PF_ASTC_12x12           =54,	// 0.89 bpp
 	PF_BC6H					=55,
 	PF_BC7					=56,
-	PF_MAX                  =57,
+	PF_R8_UINT				=57,
+	PF_MAX                  =58,
 };
 #define FOREACH_ENUM_EPIXELFORMAT(op) \
 	op(PF_Unknown) \
@@ -120,5 +121,6 @@ enum EPixelFormat
 	op(PF_ASTC_10x10) \
 	op(PF_ASTC_12x12) \
 	op(PF_BC6H) \
-	op(PF_BC7)
+	op(PF_BC7) \
+	op(PF_R8_UINT)
 
diff --git a/Engine/Source/Runtime/Engine/Public/SceneManagement.h b/Engine/Source/Runtime/Engine/Public/SceneManagement.h
--- a/Engine/Source/Runtime/Engine/Public/SceneManagement.h
+++ b/Engine/Source/Runtime/Engine/Public/SceneManagement.h
@@ -1023,6 +1023,9 @@ public:
 	 */
 	virtual void ApplyWorldOffset(FVector InOffset);
 
+	virtual float GetMaxDrawDistance() const { return 0.0f; }
+	virtual float GetFadeRange() const { return 0.0f; }
+
 protected:
 
 	friend class FScene;
diff --git a/Engine/Source/Runtime/Engine/Public/SceneUtils.h b/Engine/Source/Runtime/Engine/Public/SceneUtils.h
--- a/Engine/Source/Runtime/Engine/Public/SceneUtils.h
+++ b/Engine/Source/Runtime/Engine/Public/SceneUtils.h
@@ -147,6 +147,95 @@
 
 #endif
 
+// GPU stats
+#if STATS
+#define HAS_GPU_STATS 1
+#else
+#define HAS_GPU_STATS 0
+#endif
+
+#if HAS_GPU_STATS
+#define SCOPED_GPU_STAT(RHICmdList, Stat) FScopedGPUStatEvent PREPROCESSOR_JOIN(GPUStatEvent_##Stat,__LINE__); PREPROCESSOR_JOIN(GPUStatEvent_##Stat,__LINE__).Begin(RHICmdList, GET_STATID( Stat ) );
+#define GPU_STATS_UPDATE(RHICmdList) FRealtimeGPUProfiler::Get()->Update(RHICmdList);
+#else
+#define SCOPED_GPU_STAT(RHICmdList, Stat) 
+#define GPU_STATS_UPDATE(RHICmdList) 
+#endif
+
+#if HAS_GPU_STATS
+
+class FRealtimeGPUProfilerEvent;
+class FRealtimeGPUProfilerFrame;
+
+/**
+* FRealtimeGPUProfiler class. This manages recording and reporting all for GPU stats
+*/
+class FRealtimeGPUProfiler
+{
+	static FRealtimeGPUProfiler* Instance;
+public:
+	// Singleton interface
+	static ENGINE_API FRealtimeGPUProfiler* Get();
+
+	/** Per-frame update */
+	ENGINE_API void Update(FRHICommandListImmediate& RHICmdList);
+
+	/** Final cleanup */
+	ENGINE_API void Release();
+
+	/** Push/pop events */
+	void PushEvent(FRHICommandListImmediate& RHICmdList, TStatId StatId);
+	void PopEvent(FRHICommandListImmediate& RHICmdList);
+
+private:
+	FRealtimeGPUProfiler();
+	void UpdateStats(FRHICommandListImmediate& RHICmdList);
+
+	/** Ringbuffer of profiler frames */
+	TArray<FRealtimeGPUProfilerFrame*> Frames;
+
+	int32 WriteBufferIndex;
+	int32 ReadBufferIndex;
+	uint32 WriteFrameNumber;
+};
+
+/**
+* Class that logs GPU Stat events for the realtime GPU profiler
+*/
+class FScopedGPUStatEvent
+{
+	/** Cmdlist to push onto. */
+	FRHICommandListImmediate* RHICmdList;
+
+	/** The stat event used to record timings */
+	FRealtimeGPUProfilerEvent* RealtimeGPUProfilerEvent;
+
+public:
+	/** Default constructor, initializing all member variables. */
+	FORCEINLINE FScopedGPUStatEvent()
+		: RHICmdList(nullptr)
+		, RealtimeGPUProfilerEvent(nullptr)
+	{}
+
+	/**
+	* Terminate the event based upon scope
+	*/
+	FORCEINLINE ~FScopedGPUStatEvent()
+	{
+		if (RHICmdList)
+		{
+			End();
+		}
+	}
+
+	/**
+	* Start/Stop functions for timer stats
+	*/
+	ENGINE_API void Begin(FRHICommandList& InRHICmdList, TStatId StatID);
+	ENGINE_API void End();
+};
+#endif // HAS_GPU_STATS
+
 /** True if HDR is enabled for the mobile renderer. */
 ENGINE_API bool IsMobileHDR();
 
diff --git a/Engine/Source/Runtime/Launch/Private/LaunchEngineLoop.cpp b/Engine/Source/Runtime/Launch/Private/LaunchEngineLoop.cpp
--- a/Engine/Source/Runtime/Launch/Private/LaunchEngineLoop.cpp
+++ b/Engine/Source/Runtime/Launch/Private/LaunchEngineLoop.cpp
@@ -201,6 +201,9 @@ static TScopedPointer<FOutputDeviceTestExit> GScopedTestExit;
 #if WITH_ENGINE
 static void RHIExitAndStopRHIThread()
 {
+#if HAS_GPU_STATS
+	FRealtimeGPUProfiler::Get()->Release();
+#endif
 	RHIExit();
 
 	// Stop the RHI Thread
diff --git a/Engine/Source/Runtime/Niagara/Classes/NiagaraConstants.h b/Engine/Source/Runtime/Niagara/Classes/NiagaraConstants.h
--- a/Engine/Source/Runtime/Niagara/Classes/NiagaraConstants.h
+++ b/Engine/Source/Runtime/Niagara/Classes/NiagaraConstants.h
@@ -17,10 +17,11 @@ namespace NiagaraConstants
 		NumBuiltinConstants
 	} ConstantDef;
 
-	FName ConstantNames[] = {
+	FName GConstantNames[] = {
+		"Undefined",
 		"Delta Time",
-		"Emitter position",
-		"Emitter age",
+		"Emitter Position",
+		"Emitter Age",
 		"Emitter X Axis",
 		"Emitter Y Axis",
 		"Emitter Z Axis",
@@ -28,3 +29,6 @@ namespace NiagaraConstants
 	};
 };
 
+#define BUILTIN_CONST_DELTATIME  NiagaraConstants::GConstantNames[NiagaraConstants::DeltaSeconds]
+#define BUILTIN_CONST_EMITTERPOS  NiagaraConstants::GConstantNames[NiagaraConstants::EmitterPosition]
+#define BUILTIN_CONST_EMITTERAGE  NiagaraConstants::GConstantNames[NiagaraConstants::EmitterAge]
\ No newline at end of file
diff --git a/Engine/Source/Runtime/Niagara/Classes/NiagaraDataSet.h b/Engine/Source/Runtime/Niagara/Classes/NiagaraDataSet.h
--- a/Engine/Source/Runtime/Niagara/Classes/NiagaraDataSet.h
+++ b/Engine/Source/Runtime/Niagara/Classes/NiagaraDataSet.h
@@ -3,7 +3,153 @@
 
 #include "NiagaraCommon.h"
 
-/** Storage for Niagara data. */
+const FName BUILTIN_VAR_PARTICLEAGE = FName(TEXT("Age"));
+
+#define SEPARATE_PER_VARIABLE_DATA 1
+
+
+
+
+/** Niagara Data Buffer
+  * Holds a number of elements of data of a specific type; set Target to determine
+  * whether CPU or GPU side
+  */
+class FNiagaraDataBuffer
+{
+public:
+	FNiagaraDataBuffer()
+		: DataType(ENiagaraDataType::Vector)
+		, ElementType(ENiagaraType::Float)
+		, Target(ENiagaraSimTarget::CPUSim)
+	{
+	}
+
+	ENiagaraSimTarget GetTarget()
+	{
+		return Target;
+	}
+
+	// when changing sim target at runtime, make sure to Reset/Alloc
+	// generally only happens when changing sim type on emitters, case components are usually re-registered and no further work is needed
+	void SetTarget(ENiagaraSimTarget InTarget)
+	{
+		Target = InTarget;
+	}
+
+	void Allocate(int32 NumElements)
+	{
+		uint32 SizePerElement = GetNiagaraBytesPerElement(DataType, ElementType);
+		uint32 SizeInBytes = NumElements * SizePerElement;
+
+		if (Target == CPUSim)
+		{
+			CPUBuffer.AddUninitialized(SizeInBytes);
+		}
+		else if (Target == GPUComputeSim)
+		{
+			FRHIResourceCreateInfo CreateInfo;
+			CreateInfo.ClearValueBinding = FClearValueBinding::None;
+			CreateInfo.BulkData = nullptr;
+			GPUBuffer = RHICreateStructuredBuffer(SizePerElement, SizeInBytes, BUF_ShaderResource, CreateInfo);
+		}
+	}
+
+	void Reset(uint32 NumElements)
+	{
+		if (Target == CPUSim)
+		{
+			CPUBuffer.Reset(NumElements);
+		}
+		else if (Target == GPUComputeSim)
+		{
+			FRHIResourceCreateInfo CreateInfo;
+			uint32 SizePerElement = GetNiagaraBytesPerElement(DataType, ElementType);
+			uint32 SizeInBytes = NumElements * SizePerElement;
+			GPUBuffer = RHICreateStructuredBuffer(SizePerElement, SizeInBytes, BUF_ShaderResource, CreateInfo);
+		}
+	}
+
+	const TArray<FVector4>& GetCPUBuffer() const
+	{
+		return CPUBuffer;
+	}
+
+	TArray<FVector4>& GetCPUBuffer() 
+	{
+		return CPUBuffer;
+	}
+
+
+private:
+	ENiagaraDataType DataType;
+	ENiagaraType ElementType;
+	FStructuredBufferRHIRef GPUBuffer;
+	TArray<FVector4> CPUBuffer;
+	ENiagaraSimTarget Target;
+};
+
+
+
+/** Storage for a single particle attribute for all particles of one emitter 
+  * We double buffer at this level, so we can avoid copying or explicitly dealing
+  * with unused attributes.
+  */
+class FNiagaraVariableData
+{
+public:
+	FNiagaraVariableData()
+		: CurrentBuffer(0)
+		, bPassThroughEnabled(false)
+	{}
+	
+	TArray<FVector4> &GetCurrentBuffer() { return DataBuffers[CurrentBuffer].GetCPUBuffer(); }
+	TArray<FVector4> &GetPrevBuffer() { return DataBuffers[CurrentBuffer ^ 0x1].GetCPUBuffer(); }
+
+	const TArray<FVector4> &GetCurrentBuffer() const { return DataBuffers[CurrentBuffer].GetCPUBuffer(); }
+	const TArray<FVector4> &GetPrevBuffer() const { return DataBuffers[CurrentBuffer ^ 0x1].GetCPUBuffer(); }
+
+	void Allocate(uint32 NumElements)
+	{
+		//DataBuffers[CurrentBuffer].AddUninitialized(NumElements);
+		DataBuffers[CurrentBuffer].Allocate(NumElements);
+	}
+
+	void Reset(uint32 NumElements)
+	{
+		DataBuffers[CurrentBuffer].Reset(NumElements);
+	}
+
+	void SwapBuffers()
+	{ 
+		// Passthrough is usually enabled when an attribute is unused during VM execution; 
+		// in this case, we don't swap buffers, to avoid needing to copy source to destination
+		if (!bPassThroughEnabled)
+		{
+			CurrentBuffer ^= 0x1;
+		}
+	}
+
+	// Passthrough enabled means no double buffering; to be used for unused attributes only
+	void EnablePassThrough() { bPassThroughEnabled = true; }
+	void DisablePassThrough() { bPassThroughEnabled = false; }
+
+
+private:
+	ENiagaraDataType Type;
+	uint8 CurrentBuffer;
+	FNiagaraDataBuffer DataBuffers[2];
+	FStructuredBufferRHIRef GPUDataBuffers[2];
+
+	bool bPassThroughEnabled;
+};
+
+
+
+
+/** Storage for Niagara data. 
+  * Holds a set of FNiagaraVariableData along with mapping from FNiagaraVariableInfos
+  * Generally, this is the entire set of particle attributes for one emitter
+  */
 class FNiagaraDataSet
 {
 public:
@@ -27,12 +173,20 @@ public:
 	{
 		if (bReset)
 		{
-			DataBuffers[CurrentBuffer].Reset(NumExpectedInstances * VariableMap.Num());
-			DataBuffers[CurrentBuffer].AddUninitialized(NumExpectedInstances * VariableMap.Num());
+			VariableDataArray.SetNumZeroed(VariableMap.Num());
+			for (auto &Data : VariableDataArray)
+			{
+				Data.Reset(NumExpectedInstances);
+				Data.Allocate(NumExpectedInstances);
+			}
 		}
 		else
 		{
-			DataBuffers[CurrentBuffer].SetNumUninitialized(NumExpectedInstances * VariableMap.Num());
+			VariableDataArray.SetNumUninitialized(VariableMap.Num());
+			for (auto &Data : VariableDataArray)
+			{
+				Data.Allocate(NumExpectedInstances);
+			}
 		}
 
 		DataAllocation[CurrentBuffer] = NumExpectedInstances;
@@ -46,32 +200,30 @@ public:
 		DataAllocation[0] = 0;
 		DataAllocation[1] = 0;
 		CurrentBuffer = 0;
-		DataBuffers[0].Empty();
-		DataBuffers[1].Empty();
 	}
 
 	FVector4* GetVariableData(const FNiagaraVariableInfo& VariableID, uint32 StartParticle = 0)
 	{
 		const uint32* Offset = VariableMap.Find(VariableID);
-		return Offset ? DataBuffers[CurrentBuffer].GetData() + (*Offset * DataAllocation[CurrentBuffer]) + StartParticle : NULL;
+		return Offset ? VariableDataArray[*Offset].GetCurrentBuffer().GetData() + StartParticle : NULL;
 	}
 
 	const FVector4* GetVariableData(const FNiagaraVariableInfo& VariableID, uint32 StartParticle = 0)const
 	{
 		const uint32* Offset = VariableMap.Find(VariableID);
-		return Offset ? DataBuffers[CurrentBuffer].GetData() + (*Offset * DataAllocation[CurrentBuffer]) + StartParticle : NULL;
+		return Offset ? VariableDataArray[*Offset].GetCurrentBuffer().GetData() + StartParticle : NULL;
 	}
 
 	FVector4* GetPrevVariableData(const FNiagaraVariableInfo& VariableID, uint32 StartParticle = 0)
 	{
 		const uint32* Offset = VariableMap.Find(VariableID);
-		return Offset ? DataBuffers[CurrentBuffer ^ 0x1].GetData() + (*Offset * DataAllocation[CurrentBuffer ^ 0x1]) + StartParticle : NULL;
+		return Offset ? VariableDataArray[*Offset].GetPrevBuffer().GetData() + StartParticle : NULL;
 	}
 
 	const FVector4* GetPrevVariableData(const FNiagaraVariableInfo& VariableID, uint32 StartParticle = 0)const
 	{
 		const uint32* Offset = VariableMap.Find(VariableID);
-		return Offset ? DataBuffers[CurrentBuffer ^ 0x1].GetData() + (*Offset * DataAllocation[CurrentBuffer ^ 0x1]) + StartParticle : NULL;
+		return Offset ? VariableDataArray[*Offset].GetPrevBuffer().GetData() + StartParticle : NULL;
 	}
 
 	void GetVariableData(const FNiagaraVariableInfo& VariableID, FVector4*& PrevBuffer, FVector4*& CurrBuffer, uint32 StartParticle)
@@ -79,8 +231,8 @@ public:
 		const uint32* Offset = VariableMap.Find(VariableID);
 		if (Offset)
 		{
-			PrevBuffer = DataBuffers[CurrentBuffer ^ 0x1].GetData() + (*Offset * DataAllocation[CurrentBuffer ^ 0x1]) + StartParticle;
-			CurrBuffer = DataBuffers[CurrentBuffer].GetData() + (*Offset * DataAllocation[CurrentBuffer]) + StartParticle;
+			PrevBuffer = VariableDataArray[*Offset].GetPrevBuffer().GetData() + StartParticle;
+			CurrBuffer = VariableDataArray[*Offset].GetCurrentBuffer().GetData() + StartParticle;
 		}
 		else
 		{
@@ -89,6 +241,14 @@ public:
 		}
 	}
 
+
+	FNiagaraVariableData &GetVariableBuffer(const FNiagaraVariableInfo& VariableID)
+	{
+		const uint32* VarIdx = VariableMap.Find(VariableID);
+		return VariableDataArray[*VarIdx];
+	}
+
+
 	bool HasAttriubte(const FNiagaraVariableInfo& VariableID)
 	{
 		return VariableMap.Find(VariableID) != NULL;
@@ -123,6 +283,7 @@ public:
 		}
 	}
 
+
 	const TMap<FNiagaraVariableInfo, uint32>& GetVariables()const { return VariableMap; }
 	int GetNumVariables()				{ return VariableMap.Num(); }
 
@@ -133,10 +294,17 @@ public:
 	void SetNumInstances(uint32 Num)	{ NumInstances[CurrentBuffer] = Num; }
 	void SwapBuffers()					{ CurrentBuffer ^= 0x1; }
 
-	FVector4* GetCurrentBuffer()		{ return DataBuffers[CurrentBuffer].GetData(); }
-	FVector4* GetPreviousBuffer()		{ return DataBuffers[CurrentBuffer ^ 0x1].GetData(); }
+	FVector4* GetCurrentBuffer(uint32 AttrIndex) { return VariableDataArray[AttrIndex].GetCurrentBuffer().GetData(); }
+	FVector4* GetPreviousBuffer(uint32 AttrIndex) { return VariableDataArray[AttrIndex].GetPrevBuffer().GetData(); }
 
-	int GetBytesUsed()	{ return (DataBuffers[0].Num() + DataBuffers[1].Num()) * 16 + VariableMap.Num() * 4; }
+	int GetBytesUsed()	
+	{ 
+		if (VariableDataArray.Num() == 0)
+		{
+			return 0;
+		}
+		return (VariableDataArray[0].GetCurrentBuffer().Num() + VariableDataArray[0].GetPrevBuffer().Num()) * 16 + VariableMap.Num() * 4; 
+	}
 
 	FNiagaraDataSetID GetID()const { return ID; }
 	void SetID(FNiagaraDataSetID InID) { ID = InID; }
@@ -145,6 +313,11 @@ public:
 	{
 		SwapBuffers();
 
+		for (auto &VarData : VariableDataArray)
+		{
+			VarData.SwapBuffers();
+		}
+
 		if (ID.Type == ENiagaraDataSetType::Event)
 		{
 			SetNumInstances(0);
@@ -159,10 +332,11 @@ public:
 
 private:
 	uint32 CurrentBuffer;
-	uint32 NumInstances[2];
 	uint32 DataAllocation[2];
-	TArray<FVector4> DataBuffers[2];
+	uint32 NumInstances[2];
+
 	TMap<FNiagaraVariableInfo, uint32> VariableMap;
 	FNiagaraDataSetID ID;
+	TArray<FNiagaraVariableData> VariableDataArray;
 };
 
diff --git a/Engine/Source/Runtime/Niagara/Classes/NiagaraEffect.h b/Engine/Source/Runtime/Niagara/Classes/NiagaraEffect.h
--- a/Engine/Source/Runtime/Niagara/Classes/NiagaraEffect.h
+++ b/Engine/Source/Runtime/Niagara/Classes/NiagaraEffect.h
@@ -47,9 +47,6 @@ public:
 	virtual void PostLoad()override;
 	//End UObject Interface
 private:
-	UPROPERTY()
-	TArray<FDeprecatedNiagaraEmitterProperties>EmitterPropsSerialized_DEPRECATED;
-
 	UPROPERTY()
 	TArray<UNiagaraEmitterProperties*> EmitterProps;
 };
@@ -122,17 +119,17 @@ public:
 
 	void SetConstant(FNiagaraVariableInfo ID, const float Value)
 	{
-		Constants.SetOrAdd(ID, Value);
+		InstanceConstants.SetOrAdd(ID.Name, Value);
 	}
 
 	void SetConstant(FNiagaraVariableInfo ID, const FVector4& Value)
 	{
-		Constants.SetOrAdd(ID, Value);
+		InstanceConstants.SetOrAdd(ID.Name, Value);
 	}
 
 	void SetConstant(FNiagaraVariableInfo ID, const FMatrix& Value)
 	{
-		Constants.SetOrAdd(ID, Value);
+		InstanceConstants.SetOrAdd(ID.Name, Value);
 	}
 	
 	void Tick(float DeltaSeconds);
@@ -162,7 +159,7 @@ private:
 	float Age;
 
 	/** Local constant table. */
-	FNiagaraConstantMap Constants;
+	FNiagaraConstants InstanceConstants;
 	
 	TMap<FNiagaraDataSetID, FNiagaraDataSet> ExternalEvents;
 
diff --git a/Engine/Source/Runtime/Niagara/Classes/NiagaraEmitterProperties.h b/Engine/Source/Runtime/Niagara/Classes/NiagaraEmitterProperties.h
--- a/Engine/Source/Runtime/Niagara/Classes/NiagaraEmitterProperties.h
+++ b/Engine/Source/Runtime/Niagara/Classes/NiagaraEmitterProperties.h
@@ -102,40 +102,6 @@ struct FNiagaraEmitterScriptProperties
 	NIAGARA_API void Init(UNiagaraEmitterProperties* EmitterProps);
 };
 
-//This struct now only exists for backwards compatibility and should be removed once effects are updated.
-USTRUCT()
-struct FDeprecatedNiagaraEmitterProperties
-{
-	GENERATED_USTRUCT_BODY()
-public:
-	UPROPERTY()
-	FString Name;
-	UPROPERTY()
-	bool bIsEnabled;
-	UPROPERTY()
-	float SpawnRate;
-	UPROPERTY()
-	UNiagaraScript *UpdateScript;
-	UPROPERTY()
-	UNiagaraScript *SpawnScript;
-	UPROPERTY()
-	UMaterial *Material;
-	UPROPERTY()
-	TEnumAsByte<EEmitterRenderModuleType> RenderModuleType;
-	UPROPERTY()
-	float StartTime;
-	UPROPERTY()
-	float EndTime;
-	UPROPERTY()
-	class UNiagaraEffectRendererProperties *RendererProperties;
-	UPROPERTY()
-	FNiagaraConstantMap ExternalConstants;		// these are the update script constants from the effect editor; will be added to the emitter's constant map
-	UPROPERTY()
-	FNiagaraConstantMap ExternalSpawnConstants;		// these are the spawn script constants from the effect editor; will be added to the emitter's constant map
-	UPROPERTY()
-	int32 NumLoops;
-};
-
 /** 
  *	UNiagaraEmitterProperties stores the attributes of an FNiagaraSimulation
  *	that need to be serialized and are used for its initialization 
@@ -146,7 +112,6 @@ class UNiagaraEmitterProperties : public UObject
 	GENERATED_UCLASS_BODY()
 public:
 	NIAGARA_API void Init();
-	void InitFromOldStruct(FDeprecatedNiagaraEmitterProperties& OldStruct);
 
 	//Begin UObject Interface
 #if WITH_EDITOR
diff --git a/Engine/Source/Runtime/Niagara/Classes/NiagaraScript.h b/Engine/Source/Runtime/Niagara/Classes/NiagaraScript.h
--- a/Engine/Source/Runtime/Niagara/Classes/NiagaraScript.h
+++ b/Engine/Source/Runtime/Niagara/Classes/NiagaraScript.h
@@ -20,6 +20,8 @@ enum class EUnusedAttributeBehaviour : uint8
 	None,
 	/** The memory for the attribute is set to NIAGARA_INVALID_MEMORY. */
 	MarkInvalid, 
+	/** The attribute is passed through without double buffering */
+	PassThrough,
 };
 
 USTRUCT()
diff --git a/Engine/Source/Runtime/Niagara/Classes/NiagaraSimulation.h b/Engine/Source/Runtime/Niagara/Classes/NiagaraSimulation.h
--- a/Engine/Source/Runtime/Niagara/Classes/NiagaraSimulation.h
+++ b/Engine/Source/Runtime/Niagara/Classes/NiagaraSimulation.h
@@ -15,6 +15,7 @@ DECLARE_CYCLE_STAT(TEXT("Tick"), STAT_NiagaraTick, STATGROUP_Niagara);
 DECLARE_CYCLE_STAT(TEXT("Simulate"), STAT_NiagaraSimulate, STATGROUP_Niagara);
 DECLARE_CYCLE_STAT(TEXT("Spawn"), STAT_NiagaraSpawn, STATGROUP_Niagara);
 DECLARE_CYCLE_STAT(TEXT("Kill"), STAT_NiagaraKill, STATGROUP_Niagara);
+DECLARE_CYCLE_STAT(TEXT("Register Setup"), STAT_NiagaraSimRegisterSetup, STATGROUP_Niagara);
 
 /**
 * A Niagara particle simulation.
@@ -34,14 +35,15 @@ public:
 
 	void PreTick();
 	void Tick(float DeltaSeconds);
+	void TickEvents(float DeltaSeconds);
 
+	void KillParticles();
 
 	FBox GetBounds() const { return CachedBounds; }
 
 	FNiagaraDataSet &GetData()	{ return Data; }
 
-	void SetConstants(const FNiagaraConstantMap &InMap)	{ Constants = InMap; }
-	FNiagaraConstantMap &GetConstants()	{ return Constants; }
+	FNiagaraConstants &GetConstants()	{ return ExternalConstants; }
 
 	NiagaraEffectRenderer *GetEffectRenderer()	{ return EffectRenderer; }
 	
@@ -92,7 +94,8 @@ private:
 	ENiagaraTickState TickState;
 	
 	/** Local constant set. */
-	FNiagaraConstantMap Constants;
+	FNiagaraConstants ExternalConstants;
+
 	/** particle simulation data */
 	FNiagaraDataSet Data;
 	/** Keep partial particle spawns from last frame */
@@ -135,14 +138,10 @@ private:
 	/** Util to move a particle */
 	void MoveParticleToIndex(int32 SrcIndex, int32 DestIndex)
 	{
-		FVector4 *SrcPtr = Data.GetCurrentBuffer() + SrcIndex;
-		FVector4 *DestPtr = Data.GetCurrentBuffer() + DestIndex;
-
 		for (int32 AttrIndex = 0; AttrIndex < Data.GetNumVariables(); AttrIndex++)
 		{
-			*DestPtr = *SrcPtr;
-			DestPtr += Data.GetDataAllocation();
-			SrcPtr += Data.GetDataAllocation();
+			FVector4 *AttrPtr = Data.GetCurrentBuffer(AttrIndex);
+			*(AttrPtr+DestIndex) = *(AttrPtr+SrcIndex);
 		}
 	}
 
diff --git a/Engine/Source/Runtime/Niagara/NiagaraConstantSet.h b/Engine/Source/Runtime/Niagara/NiagaraConstantSet.h
--- a/Engine/Source/Runtime/Niagara/NiagaraConstantSet.h
+++ b/Engine/Source/Runtime/Niagara/NiagaraConstantSet.h
@@ -7,56 +7,6 @@
 
 struct FNiagaraConstants;
 
-/** Runtime storage for emitter overridden constants. Only now a USTRUCT still to support BC */
-USTRUCT()
-struct FNiagaraConstantMap
-{
-	GENERATED_USTRUCT_BODY()
-private:
-	TMap<FNiagaraVariableInfo, float> ScalarConstants;
-	TMap<FNiagaraVariableInfo, FVector4> VectorConstants;
-	TMap<FNiagaraVariableInfo, FMatrix> MatrixConstants;
-
-public:
-	FNiagaraConstantMap()
-	{
-	}
-
-	void Empty();
-
-	void SetOrAdd(FNiagaraVariableInfo ID, float Sc);
-	void SetOrAdd(FNiagaraVariableInfo ID, const FVector4& Vc);
-	void SetOrAdd(FNiagaraVariableInfo ID, const FMatrix& Mc);
-	void SetOrAdd(FName Name, float Sc);
-	void SetOrAdd(FName Name, const FVector4& Vc);
-	void SetOrAdd(FName Name, const FMatrix& Mc);
-
-	float* FindScalar(FNiagaraVariableInfo ID);
-	FVector4* FindVector(FNiagaraVariableInfo ID);
-	FMatrix* FindMatrix(FNiagaraVariableInfo ID);
-	const float* FindScalar(FNiagaraVariableInfo ID)const;
-	const FVector4* FindVector(FNiagaraVariableInfo ID)const;
-	const FMatrix* FindMatrix(FNiagaraVariableInfo ID)const;
-	float* FindScalar(FName Name);
-	FVector4* FindVector(FName Name);
-	FMatrix* FindMatrix(FName Name);
-	const float* FindScalar(FName Name)const;
-	const FVector4* FindVector(FName Name)const;
-	const FMatrix* FindMatrix(FName Name)const;
-
-	void Merge(FNiagaraConstants &InConstants);
-	void Merge(FNiagaraConstantMap &InMap);
-
-	const TMap<FNiagaraVariableInfo, float> &GetScalarConstants()	const;
-	const TMap<FNiagaraVariableInfo, FVector4> &GetVectorConstants()	const;
-	const TMap<FNiagaraVariableInfo, FMatrix> &GetMatrixConstants()	const;
-
-
-	virtual bool Serialize(FArchive &Ar)
-	{
-		return true;
-	}
-};
 
 USTRUCT()
 struct FNiagaraConstantBase
@@ -76,15 +26,15 @@ struct FNiagaraConstantBase
 };
 
 USTRUCT()
-struct FNiagaraConstants_Float : public FNiagaraConstantBase
+struct FNiagaraConstant_Float : public FNiagaraConstantBase
 {
 	GENERATED_USTRUCT_BODY()
 
-	FNiagaraConstants_Float()
+	FNiagaraConstant_Float()
 	: Value(0.0f)
 	{}
 
-	FNiagaraConstants_Float(FName InName, float InValue)
+	FNiagaraConstant_Float(FName InName, float InValue)
 	: FNiagaraConstantBase(InName)
 	, Value(InValue)
 	{}
@@ -94,15 +44,15 @@ struct FNiagaraConstants_Float : public FNiagaraConstantBase
 };
 
 USTRUCT()
-struct FNiagaraConstants_Vector : public FNiagaraConstantBase
+struct FNiagaraConstant_Vector : public FNiagaraConstantBase
 {
 	GENERATED_USTRUCT_BODY()
 
-	FNiagaraConstants_Vector()
+	FNiagaraConstant_Vector()
 	: Value(FVector4(1.0f, 1.0f, 1.0f, 1.0f))
 	{}
 
-	FNiagaraConstants_Vector(FName InName, const FVector4& InValue)
+	FNiagaraConstant_Vector(FName InName, const FVector4& InValue)
 	: FNiagaraConstantBase(InName)
 	, Value(InValue)
 	{}
@@ -112,15 +62,15 @@ struct FNiagaraConstants_Vector : public FNiagaraConstantBase
 };
 
 USTRUCT()
-struct FNiagaraConstants_Matrix : public FNiagaraConstantBase
+struct FNiagaraConstant_Matrix : public FNiagaraConstantBase
 {
 	GENERATED_USTRUCT_BODY()
 
-	FNiagaraConstants_Matrix()
+	FNiagaraConstant_Matrix()
 	: Value(FMatrix::Identity)
 	{}
 
-	FNiagaraConstants_Matrix(FName InName, const FMatrix& InValue)
+	FNiagaraConstant_Matrix(FName InName, const FMatrix& InValue)
 	: FNiagaraConstantBase(InName)
 	, Value(InValue)
 	{}
@@ -129,28 +79,6 @@ struct FNiagaraConstants_Matrix : public FNiagaraConstantBase
 	FMatrix Value;
 };
 
-//Dummy struct used to serialize in the old layout of FNiagaraConstants
-//This should be removed once everyone has recompiled and saved their scripts.
-USTRUCT()
-struct FDeprecatedNiagaraConstants
-{
-	GENERATED_USTRUCT_BODY()
-
-	//DEPRECATED PROPERTIES. REMOVE SOON!
-	UPROPERTY()
-	TArray<FNiagaraVariableInfo> ScalarConstantsInfo_DEPRECATED;
-	UPROPERTY()
-	TArray<FNiagaraVariableInfo> VectorConstantsInfo_DEPRECATED;
-	UPROPERTY()
-	TArray<FNiagaraVariableInfo> MatrixConstantsInfo_DEPRECATED;
-	UPROPERTY()
-	TArray<float> ScalarConstants_DEPRECATED;
-	UPROPERTY()
-	TArray<FVector4> VectorConstants_DEPRECATED;
-	UPROPERTY()
-	TArray<FMatrix> MatrixConstants_DEPRECATED;
-};
-
 USTRUCT()
 struct FNiagaraConstants
 {
@@ -159,11 +87,11 @@ struct FNiagaraConstants
 public:
 
 	UPROPERTY(EditAnywhere, EditFixedSize, Category = "Constant")
-	TArray<FNiagaraConstants_Float> ScalarConstants;
+	TArray<FNiagaraConstant_Float> ScalarConstants;
 	UPROPERTY(EditAnywhere, EditFixedSize, Category = "Constant")
-	TArray<FNiagaraConstants_Vector> VectorConstants;
+	TArray<FNiagaraConstant_Vector> VectorConstants;
 	UPROPERTY(EditAnywhere, EditFixedSize, Category = "Constant")
-	TArray<FNiagaraConstants_Matrix> MatrixConstants;
+	TArray<FNiagaraConstant_Matrix> MatrixConstants;
 
 	NIAGARA_API void Empty();
 	NIAGARA_API void GetScalarConstant(int32 Index, float& OutValue, FNiagaraVariableInfo& OutInfo);
@@ -176,20 +104,15 @@ public:
 
 	NIAGARA_API void Init(class UNiagaraEmitterProperties* EmitterProps, struct FNiagaraEmitterScriptProperties* ScriptProps);
 
-	/**
-	Copies constants in from the script as usual but allows overrides from old data in an FNiagaraConstantMap.
-	Only used for BC. DO NOT USE.
-	*/
-	NIAGARA_API void InitFromOldMap(FNiagaraConstantMap& OldMap);
-
 	/** Fills the entire constants set into the constant table. */
 	void AppendToConstantsTable(TArray<FVector4>& ConstantsTable)const;
 
 	/** 
 	Fills only selected constants into the table. In the order they appear in the array of passed names not the order they appear in the set. 
-	Checks the passed map for entires to superceed the default values in the set.
+	Checks the passed map for entries to supersede the default values in the set.
 	*/
-	void AppendToConstantsTable(TArray<FVector4>& ConstantsTable, const FNiagaraConstantMap& Externals)const;
+	void AppendToConstantsTable(TArray<FVector4>& ConstantsTable, const FNiagaraConstants& Externals)const;
+
 
 	NIAGARA_API void SetOrAdd(const FNiagaraVariableInfo& Constant, float Value);
 	NIAGARA_API void SetOrAdd(const FNiagaraVariableInfo& Constant, const FVector4& Value);
@@ -212,10 +135,15 @@ public:
 	NIAGARA_API FVector4* FindVector(FName Name);
 	NIAGARA_API FMatrix* FindMatrix(FName Name);
 
+	NIAGARA_API const float* FindScalar(FName Name) const;
+	NIAGARA_API const FVector4* FindVector(FName Name) const;
+	NIAGARA_API const FMatrix* FindMatrix(FName Name) const;
+
+
 	/** Return the first constant used for this constant in a table. */
 	FORCEINLINE int32 GetTableIndex_Scalar(const FNiagaraVariableInfo& Constant)const
 	{ 
-		int Idx = ScalarConstants.IndexOfByPredicate([&](const FNiagaraConstants_Float& V){ return Constant.Name == V.Name; });
+		int Idx = ScalarConstants.IndexOfByPredicate([&](const FNiagaraConstant_Float& V){ return Constant.Name == V.Name; });
 		if (Idx != INDEX_NONE)
 		{
 			return Idx;
@@ -225,7 +153,7 @@ public:
 
 	FORCEINLINE int32 GetTableIndex_Vector(const FNiagaraVariableInfo& Constant)const
 	{ 
-		int Idx = VectorConstants.IndexOfByPredicate([&](const FNiagaraConstants_Vector& V){ return Constant.Name == V.Name; });
+		int Idx = VectorConstants.IndexOfByPredicate([&](const FNiagaraConstant_Vector& V){ return Constant.Name == V.Name; });
 		if (Idx != INDEX_NONE)
 		{
 			return ScalarTableSize() + Idx;
@@ -235,11 +163,28 @@ public:
 
 	FORCEINLINE int32 GetTableIndex_Matrix(const FNiagaraVariableInfo& Constant)const
 	{ 
-		int Idx = MatrixConstants.IndexOfByPredicate([&](const FNiagaraConstants_Matrix& V){ return Constant.Name == V.Name; });
+		int Idx = MatrixConstants.IndexOfByPredicate([&](const FNiagaraConstant_Matrix& V){ return Constant.Name == V.Name; });
 		if (Idx != INDEX_NONE)
 		{
 			return ScalarTableSize() + VectorTableSize() + Idx * 4;
 		}
 		return INDEX_NONE;
 	}
+
+
+	void Merge(FNiagaraConstants &InConstants)
+	{
+		for (FNiagaraConstant_Float& C : InConstants.ScalarConstants)
+		{
+			SetOrAdd(FNiagaraVariableInfo(C.Name, ENiagaraDataType::Scalar), C.Value);
+		}
+		for (FNiagaraConstant_Vector& C : InConstants.VectorConstants)
+		{
+			SetOrAdd(FNiagaraVariableInfo(C.Name, ENiagaraDataType::Vector), C.Value);
+		}
+		for (FNiagaraConstant_Matrix& C : InConstants.MatrixConstants)
+		{
+			SetOrAdd(FNiagaraVariableInfo(C.Name, ENiagaraDataType::Matrix), C.Value);
+		}
+	}
 };
diff --git a/Engine/Source/Runtime/Niagara/NiagaraScriptConstantData.h b/Engine/Source/Runtime/Niagara/NiagaraScriptConstantData.h
--- a/Engine/Source/Runtime/Niagara/NiagaraScriptConstantData.h
+++ b/Engine/Source/Runtime/Niagara/NiagaraScriptConstantData.h
@@ -12,11 +12,6 @@ struct FNiagaraScriptConstantData
 	GENERATED_USTRUCT_BODY()
 
 	/** The set of external constants for this script. */
-	UPROPERTY()
-	FDeprecatedNiagaraConstants ExternalConstants_DEPRECATED;
-	UPROPERTY()
-	FDeprecatedNiagaraConstants InternalConstants_DEPRECATED;
-
 	/** Constants driven by the system. Named New for BC reasons. Once all data is updated beyond VER_UE4_NIAGARA_DATA_OBJECT_DEV_UI_FIX. Get rid of the deprecated consts and rename the New. */
 	UPROPERTY()
 	FNiagaraConstants ExternalConstantsNew;
@@ -33,17 +28,18 @@ struct FNiagaraScriptConstantData
 
 	FNiagaraConstants& GetExternalConstants(){ return ExternalConstantsNew; }
 
-	/** Fill a constants table ready for use in an update or spawn. */
-	void FillConstantTable(const FNiagaraConstantMap& ExternalConstantsMap, TArray<FVector4>& OutConstantTable) const
+	void FillConstantTable(const FNiagaraConstants& ExternalConstants, TArray<FVector4>& OutConstantTable) const
 	{
 		//First up in the table comes the External constants in scalar->vector->matrix order.
 		//Only fills the constants actually used by the script.
 		OutConstantTable.Empty();
-		ExternalConstantsNew.AppendToConstantsTable(OutConstantTable, ExternalConstantsMap);
+		ExternalConstantsNew.AppendToConstantsTable(OutConstantTable, ExternalConstants);
 		//Next up add all the internal constants from the script.
 		InternalConstantsNew.AppendToConstantsTable(OutConstantTable);
 	}
 
+
+
 	template<typename T>
 	void SetOrAddInternal(const FNiagaraVariableInfo& Constant, const T& Value)
 	{
diff --git a/Engine/Source/Runtime/Niagara/Private/NiagaraCommon.cpp b/Engine/Source/Runtime/Niagara/Private/NiagaraCommon.cpp
new file mode 100644
--- /dev/null
+++ b/Engine/Source/Runtime/Niagara/Private/NiagaraCommon.cpp
@@ -0,0 +1,16 @@
+#include "NiagaraPrivate.h"
+#include "NiagaraCommon.h"
+
+uint32 GetNiagaraDataElementCount(ENiagaraDataType Type)
+{
+	static int32 Counts[] = { 1, 4, 16 };
+	return Counts[(uint8)Type];
+}
+
+uint32 GetNiagaraBytesPerElement(ENiagaraDataType DataType, ENiagaraType Type)
+{
+	static int32 Sizes[] = { 1, 2, 4, 4 };
+	return Sizes[(uint8)Type] * GetNiagaraDataElementCount(DataType);
+}
+
+
diff --git a/Engine/Source/Runtime/Niagara/Private/NiagaraConstantSet.cpp b/Engine/Source/Runtime/Niagara/Private/NiagaraConstantSet.cpp
--- a/Engine/Source/Runtime/Niagara/Private/NiagaraConstantSet.cpp
+++ b/Engine/Source/Runtime/Niagara/Private/NiagaraConstantSet.cpp
@@ -4,123 +4,6 @@
 #include "NiagaraConstantSet.h"
 #include "NiagaraSimulation.h"
 
-//////////////////////////////////////////////////////////////////////////
-// FNiagaraConstantMap
-
-template<>
-struct TStructOpsTypeTraits<FNiagaraConstantMap> : public TStructOpsTypeTraitsBase
-{
-	enum
-	{
-		WithSerializer = true
-	};
-};
-IMPLEMENT_STRUCT(NiagaraConstantMap);
-
-void FNiagaraConstantMap::Empty()
-{
-	ScalarConstants.Empty();
-	VectorConstants.Empty();
-	MatrixConstants.Empty();
-}
-
-void FNiagaraConstantMap::SetOrAdd(FNiagaraVariableInfo ID, float Sc)
-{
-	check(ID.Type == ENiagaraDataType::Scalar);
-	ScalarConstants.FindOrAdd(ID) = Sc;
-}
-
-void FNiagaraConstantMap::SetOrAdd(FNiagaraVariableInfo ID, const FVector4& Vc)
-{
-	check(ID.Type == ENiagaraDataType::Vector);
-	VectorConstants.FindOrAdd(ID) = Vc;
-}
-
-void FNiagaraConstantMap::SetOrAdd(FNiagaraVariableInfo ID, const FMatrix& Mc)
-{
-	check(ID.Type == ENiagaraDataType::Matrix);
-	MatrixConstants.FindOrAdd(ID) = Mc;
-}
-
-void FNiagaraConstantMap::SetOrAdd(FName Name, float Sc)
-{
-	ScalarConstants.FindOrAdd(FNiagaraVariableInfo(Name, ENiagaraDataType::Scalar)) = Sc;
-}
-
-void FNiagaraConstantMap::SetOrAdd(FName Name, const FVector4& Vc)
-{
-	VectorConstants.FindOrAdd(FNiagaraVariableInfo(Name, ENiagaraDataType::Vector)) = Vc;
-}
-
-void FNiagaraConstantMap::SetOrAdd(FName Name, const FMatrix& Mc)
-{
-	MatrixConstants.FindOrAdd(FNiagaraVariableInfo(Name, ENiagaraDataType::Matrix)) = Mc;
-}
-
-float* FNiagaraConstantMap::FindScalar(FNiagaraVariableInfo ID){ return ScalarConstants.Find(ID); }
-FVector4* FNiagaraConstantMap::FindVector(FNiagaraVariableInfo ID){ return VectorConstants.Find(ID); }
-FMatrix* FNiagaraConstantMap::FindMatrix(FNiagaraVariableInfo ID){ return MatrixConstants.Find(ID); }
-
-const float* FNiagaraConstantMap::FindScalar(FNiagaraVariableInfo ID)const { return ScalarConstants.Find(ID); }
-const FVector4* FNiagaraConstantMap::FindVector(FNiagaraVariableInfo ID)const { return VectorConstants.Find(ID); }
-const FMatrix* FNiagaraConstantMap::FindMatrix(FNiagaraVariableInfo ID)const { return MatrixConstants.Find(ID); }
-
-float* FNiagaraConstantMap::FindScalar(FName Name) { return ScalarConstants.Find(FNiagaraVariableInfo(Name, ENiagaraDataType::Scalar)); }
-FVector4* FNiagaraConstantMap::FindVector(FName Name) { return VectorConstants.Find(FNiagaraVariableInfo(Name, ENiagaraDataType::Vector)); }
-FMatrix* FNiagaraConstantMap::FindMatrix(FName Name) { return MatrixConstants.Find(FNiagaraVariableInfo(Name, ENiagaraDataType::Matrix)); }
-
-const float* FNiagaraConstantMap::FindScalar(FName Name)const  { return ScalarConstants.Find(FNiagaraVariableInfo(Name, ENiagaraDataType::Scalar)); }
-const FVector4* FNiagaraConstantMap::FindVector(FName Name)const  { return VectorConstants.Find(FNiagaraVariableInfo(Name, ENiagaraDataType::Vector)); }
-const FMatrix* FNiagaraConstantMap::FindMatrix(FName Name)const  { return MatrixConstants.Find(FNiagaraVariableInfo(Name, ENiagaraDataType::Matrix)); }
-
-void FNiagaraConstantMap::Merge(FNiagaraConstants &InConstants)
-{
-	for (FNiagaraConstants_Float& C : InConstants.ScalarConstants)
-	{
-		SetOrAdd(FNiagaraVariableInfo(C.Name, ENiagaraDataType::Scalar), C.Value);
-	}
-	for (FNiagaraConstants_Vector& C : InConstants.VectorConstants)
-	{
-		SetOrAdd(FNiagaraVariableInfo(C.Name, ENiagaraDataType::Vector), C.Value);
-	}
-	for (FNiagaraConstants_Matrix& C : InConstants.MatrixConstants)
-	{
-		SetOrAdd(FNiagaraVariableInfo(C.Name, ENiagaraDataType::Matrix), C.Value);
-	}
-}
-
-void FNiagaraConstantMap::Merge(FNiagaraConstantMap &InMap)
-{
-	for (auto MapIt = InMap.ScalarConstants.CreateIterator(); MapIt; ++MapIt)
-	{
-		SetOrAdd(MapIt.Key(), MapIt.Value());
-	}
-
-	for (auto MapIt = InMap.VectorConstants.CreateIterator(); MapIt; ++MapIt)
-	{
-		SetOrAdd(MapIt.Key(), MapIt.Value());
-	}
-
-	for (auto MapIt = InMap.MatrixConstants.CreateIterator(); MapIt; ++MapIt)
-	{
-		SetOrAdd(MapIt.Key(), MapIt.Value());
-	}
-}
-
-const TMap<FNiagaraVariableInfo, float> &FNiagaraConstantMap::GetScalarConstants()	const
-{
-	return ScalarConstants;
-}
-
-const TMap<FNiagaraVariableInfo, FVector4> &FNiagaraConstantMap::GetVectorConstants()	const
-{
-	return VectorConstants;
-}
-
-const TMap<FNiagaraVariableInfo, FMatrix> &FNiagaraConstantMap::GetMatrixConstants()	const
-{
-	return MatrixConstants;
-}
 
 //////////////////////////////////////////////////////////////////////////
 
@@ -157,30 +40,30 @@ void FNiagaraConstants::Init(UNiagaraEmitterProperties* EmitterProps, FNiagaraEm
 	//Copy the script constants, retaining any altered values.
 	const FNiagaraConstants& ScriptConstants = ScriptProps->Script->ConstantData.GetExternalConstants();
 
-	TArray<FNiagaraConstants_Float> NewScalarConstants = ScriptConstants.ScalarConstants;
+	TArray<FNiagaraConstant_Float> NewScalarConstants = ScriptConstants.ScalarConstants;
 	for (auto& Override : ScalarConstants)
 	{
-		int32 Idx = NewScalarConstants.IndexOfByPredicate([&](const FNiagaraConstants_Float& C){ return Override.Name == C.Name; });
+		int32 Idx = NewScalarConstants.IndexOfByPredicate([&](const FNiagaraConstant_Float& C){ return Override.Name == C.Name; });
 
 		if (Idx != INDEX_NONE)
 			NewScalarConstants[Idx].Value = Override.Value;
 	}
 	ScalarConstants = NewScalarConstants;
 
-	TArray<FNiagaraConstants_Vector> NewVectorConstants = ScriptConstants.VectorConstants;
+	TArray<FNiagaraConstant_Vector> NewVectorConstants = ScriptConstants.VectorConstants;
 	for (auto& Override : VectorConstants)
 	{
-		int32 Idx = NewVectorConstants.IndexOfByPredicate([&](const FNiagaraConstants_Vector& C){ return Override.Name == C.Name; });
+		int32 Idx = NewVectorConstants.IndexOfByPredicate([&](const FNiagaraConstant_Vector& C){ return Override.Name == C.Name; });
 
 		if (Idx != INDEX_NONE)
 			NewVectorConstants[Idx].Value = Override.Value;
 	}
 	VectorConstants = NewVectorConstants;
 
-	TArray<FNiagaraConstants_Matrix> NewMatrixConstants = ScriptConstants.MatrixConstants;
+	TArray<FNiagaraConstant_Matrix> NewMatrixConstants = ScriptConstants.MatrixConstants;
 	for (auto& Override : MatrixConstants)
 	{
-		int32 Idx = NewMatrixConstants.IndexOfByPredicate([&](const FNiagaraConstants_Matrix& C){ return Override.Name == C.Name; });
+		int32 Idx = NewMatrixConstants.IndexOfByPredicate([&](const FNiagaraConstant_Matrix& C){ return Override.Name == C.Name; });
 
 		if (Idx != INDEX_NONE)
 			NewMatrixConstants[Idx].Value = Override.Value;
@@ -189,45 +72,21 @@ void FNiagaraConstants::Init(UNiagaraEmitterProperties* EmitterProps, FNiagaraEm
 	MatrixConstants = NewMatrixConstants;
 }
 
-/** 
-Copies constants in from the script as usual but allows overrides from old data in an FNiagaraConstantMap.
-Only used for BC. DO NOT USE.
-*/
-void FNiagaraConstants::InitFromOldMap(FNiagaraConstantMap& OldMap)
-{
-	const TMap<FNiagaraVariableInfo, float>& OldScalarConsts = OldMap.GetScalarConstants();
-	const TMap<FNiagaraVariableInfo, FVector4>& OldVectorConsts = OldMap.GetVectorConstants();
-	const TMap<FNiagaraVariableInfo, FMatrix>& OldMatrixConsts = OldMap.GetMatrixConstants();
-
-	for (const TPair<FNiagaraVariableInfo, float> Pair : OldScalarConsts)
-	{
-		ScalarConstants.Add(FNiagaraConstants_Float(Pair.Key.Name, Pair.Value));
-	}
-	for (const TPair<FNiagaraVariableInfo, FVector4>& Pair : OldVectorConsts)
-	{
-		VectorConstants.Add(FNiagaraConstants_Vector(Pair.Key.Name, Pair.Value));
-	}
-	for (const TPair<FNiagaraVariableInfo, FMatrix>& Pair : OldMatrixConsts)
-	{
-		MatrixConstants.Add(FNiagaraConstants_Matrix(Pair.Key.Name, Pair.Value));
-	}
-}
-
 void FNiagaraConstants::AppendToConstantsTable(TArray<FVector4>& ConstantsTable)const
 {
 	int32 Idx = ConstantsTable.Num();
 	int32 NewIdx = 0;
 	ConstantsTable.AddUninitialized(ScalarTableSize() + VectorTableSize() + MatrixTableSize());
-	for (const FNiagaraConstants_Float& Sc : ScalarConstants)
+	for (const FNiagaraConstant_Float& Sc : ScalarConstants)
 	{
 		ConstantsTable[Idx++] = FVector4(Sc.Value, Sc.Value, Sc.Value, Sc.Value);
 	}
 
-	for (const FNiagaraConstants_Vector& Vc : VectorConstants)
+	for (const FNiagaraConstant_Vector& Vc : VectorConstants)
 	{
 		ConstantsTable[Idx++] = Vc.Value;
 	}
-	for (const FNiagaraConstants_Matrix& C : MatrixConstants)
+	for (const FNiagaraConstant_Matrix& C : MatrixConstants)
 	{
 		const FMatrix& Mc = C.Value;
 		ConstantsTable[Idx] = FVector4(Mc.M[0][0], Mc.M[0][1], Mc.M[0][2], Mc.M[0][3]);
@@ -238,7 +97,9 @@ void FNiagaraConstants::AppendToConstantsTable(TArray<FVector4>& ConstantsTable)
 	}
 }
 
-void FNiagaraConstants::AppendToConstantsTable(TArray<FVector4>& ConstantsTable, const FNiagaraConstantMap& Externals)const
+
+
+void FNiagaraConstants::AppendToConstantsTable(TArray<FVector4>& ConstantsTable, const FNiagaraConstants& Externals)const
 {
 	int32 Idx = ConstantsTable.Num();
 	ConstantsTable.AddUninitialized(ScalarTableSize() + VectorTableSize() + MatrixTableSize());
@@ -248,7 +109,7 @@ void FNiagaraConstants::AppendToConstantsTable(TArray<FVector4>& ConstantsTable,
 		float Value = Scalar ? *Scalar : ScalarConstants[i].Value;
 		ConstantsTable[Idx++] = FVector4(Value, Value, Value, Value);
 	}
-	
+
 	for (int32 i = 0; i < VectorConstants.Num(); ++i)
 	{
 		const FVector4* Vector = Externals.FindVector(VectorConstants[i].Name);
@@ -266,45 +127,46 @@ void FNiagaraConstants::AppendToConstantsTable(TArray<FVector4>& ConstantsTable,
 	}
 }
 
+
 void FNiagaraConstants::SetOrAdd(const FNiagaraVariableInfo& Constant, float Value)
 {
 	check(Constant.Type == ENiagaraDataType::Scalar);
-	FNiagaraConstants_Float* Const = ScalarConstants.FindByPredicate([&](const FNiagaraConstants_Float& C){ return Constant.Name == C.Name; });
+	FNiagaraConstant_Float* Const = ScalarConstants.FindByPredicate([&](const FNiagaraConstant_Float& C){ return Constant.Name == C.Name; });
 	if (Const)
 	{
 		Const->Value = Value;
 	}
 	else
 	{
-		ScalarConstants.Add(FNiagaraConstants_Float(Constant.Name, Value));
+		ScalarConstants.Add(FNiagaraConstant_Float(Constant.Name, Value));
 	}
 }
 
 void FNiagaraConstants::SetOrAdd(const FNiagaraVariableInfo& Constant, const FVector4& Value)
 {
 	check(Constant.Type == ENiagaraDataType::Vector);
-	FNiagaraConstants_Vector* Const = VectorConstants.FindByPredicate([&](const FNiagaraConstants_Vector& C){ return Constant.Name == C.Name; });
+	FNiagaraConstant_Vector* Const = VectorConstants.FindByPredicate([&](const FNiagaraConstant_Vector& C){ return Constant.Name == C.Name; });
 	if (Const)
 	{
 		Const->Value = Value;
 	}
 	else
 	{
-		VectorConstants.Add(FNiagaraConstants_Vector(Constant.Name, Value));
+		VectorConstants.Add(FNiagaraConstant_Vector(Constant.Name, Value));
 	}
 }
 
 void FNiagaraConstants::SetOrAdd(const FNiagaraVariableInfo& Constant, const FMatrix& Value)
 {
 	check(Constant.Type == ENiagaraDataType::Matrix);
-	FNiagaraConstants_Matrix* Const = MatrixConstants.FindByPredicate([&](const FNiagaraConstants_Matrix& C){ return Constant.Name == C.Name; });
+	FNiagaraConstant_Matrix* Const = MatrixConstants.FindByPredicate([&](const FNiagaraConstant_Matrix& C){ return Constant.Name == C.Name; });
 	if (Const)
 	{
 		Const->Value = Value;
 	}
 	else
 	{
-		MatrixConstants.Add(FNiagaraConstants_Matrix(Constant.Name, Value));
+		MatrixConstants.Add(FNiagaraConstant_Matrix(Constant.Name, Value));
 	}
 }
 
@@ -325,17 +187,38 @@ void FNiagaraConstants::SetOrAdd(FName Name, const FMatrix& Value)
 
 float* FNiagaraConstants::FindScalar(FName Name)
 {
-	auto* C = ScalarConstants.FindByPredicate([&](const FNiagaraConstants_Float& Constant){ return Constant.Name == Name; });
+	auto* C = ScalarConstants.FindByPredicate([&](const FNiagaraConstant_Float& Constant){ return Constant.Name == Name; });
 	return C ? &C->Value : nullptr;
 }
 FVector4* FNiagaraConstants::FindVector(FName Name)
 {
-	auto* C = VectorConstants.FindByPredicate([&](const FNiagaraConstants_Vector& Constant){ return Constant.Name == Name; });
+	auto* C = VectorConstants.FindByPredicate([&](const FNiagaraConstant_Vector& Constant){ return Constant.Name == Name; });
 	return C ? &C->Value : nullptr;
 }
 FMatrix* FNiagaraConstants::FindMatrix(FName Name)
 {
-	auto* C = MatrixConstants.FindByPredicate([&](const FNiagaraConstants_Matrix& Constant){ return Constant.Name == Name; });
+	auto* C = MatrixConstants.FindByPredicate([&](const FNiagaraConstant_Matrix& Constant){ return Constant.Name == Name; });
+	return C ? &C->Value : nullptr;
+}
+
+
+
+
+const float*  FNiagaraConstants::FindScalar(FName Name) const
+{
+	auto* C = ScalarConstants.FindByPredicate([&](const FNiagaraConstant_Float& Constant) { return Constant.Name == Name; });
+	return C ? &C->Value : nullptr;
+}
+
+const FVector4* FNiagaraConstants::FindVector(FName Name) const
+{
+	auto* C = VectorConstants.FindByPredicate([&](const FNiagaraConstant_Vector& Constant) { return Constant.Name == Name; });
+	return C ? &C->Value : nullptr;
+}
+
+const FMatrix* FNiagaraConstants::FindMatrix(FName Name) const
+{
+	auto* C = MatrixConstants.FindByPredicate([&](const FNiagaraConstant_Matrix& Constant) { return Constant.Name == Name; });
 	return C ? &C->Value : nullptr;
 }
 
diff --git a/Engine/Source/Runtime/Niagara/Private/NiagaraEffect.cpp b/Engine/Source/Runtime/Niagara/Private/NiagaraEffect.cpp
--- a/Engine/Source/Runtime/Niagara/Private/NiagaraEffect.cpp
+++ b/Engine/Source/Runtime/Niagara/Private/NiagaraEffect.cpp
@@ -46,17 +46,6 @@ void UNiagaraEffect::CreateEffectRendererProps(TSharedPtr<FNiagaraSimulation> Si
 void UNiagaraEffect::PostLoad()
 {
 	Super::PostLoad();
-
-	if (GetLinkerUE4Version() < VER_UE4_NIAGARA_DATA_OBJECT_DEV_UI_FIX)
-	{
-		int32 NumEmitters = EmitterPropsSerialized_DEPRECATED.Num();
-		for (int32 i = 0; i < NumEmitters; ++i)
-		{
-			UNiagaraEmitterProperties* NewProps = NewObject<UNiagaraEmitterProperties>(this);
-			NewProps->InitFromOldStruct(EmitterPropsSerialized_DEPRECATED[i]);
-			AddEmitterProperties(NewProps);
-		}
-	}
 }
 
 //////////////////////////////////////////////////////////////////////////
@@ -125,7 +114,7 @@ void FNiagaraEffectInstance::Tick(float DeltaSeconds)
 		//TODO - Handle constants better. Like waaaay better.
 		it->GetConstants().Merge(it->GetProperties()->SpawnScriptProps.ExternalConstants);
 		it->GetConstants().Merge(it->GetProperties()->UpdateScriptProps.ExternalConstants);
-		it->GetConstants().Merge(Constants);
+		it->GetConstants().Merge(InstanceConstants);
 
 		it->PreTick();
 	}
diff --git a/Engine/Source/Runtime/Niagara/Private/NiagaraEffectRenderer.cpp b/Engine/Source/Runtime/Niagara/Private/NiagaraEffectRenderer.cpp
--- a/Engine/Source/Runtime/Niagara/Private/NiagaraEffectRenderer.cpp
+++ b/Engine/Source/Runtime/Niagara/Private/NiagaraEffectRenderer.cpp
@@ -166,9 +166,13 @@ void NiagaraEffectRendererSprites::GetDynamicMeshElements(const TArray<const FSc
 				PerViewUniformParameters.RotationBias = 0.0f;
 				PerViewUniformParameters.RotationScale = 1.0f;
 				PerViewUniformParameters.TangentSelector = FVector4(0.0f, 0.0f, 0.0f, 1.0f);
-
-
 				PerViewUniformParameters.InvDeltaSeconds = 30.0f;
+				PerViewUniformParameters.NormalsType = 0.0f;
+				PerViewUniformParameters.NormalsSphereCenter = FVector4(0.0f, 0.0f, 0.0f, 1.0f);
+				PerViewUniformParameters.NormalsCylinderUnitDirection = FVector4(0.0f, 0.0f, 1.0f, 0.0f);
+				PerViewUniformParameters.PivotOffset = FVector2D(-0.5f, -0.5f);
+				PerViewUniformParameters.MacroUVParameters = FVector4(0.0f, 0.0f, 1.0f, 1.0f);
+
 				if (Properties)
 				{
 					PerViewUniformParameters.SubImageSize = FVector4(Properties->SubImageInfo.X, Properties->SubImageInfo.Y, 1.0f / Properties->SubImageInfo.X, 1.0f / Properties->SubImageInfo.Y);
@@ -180,11 +184,6 @@ void NiagaraEffectRendererSprites::GetDynamicMeshElements(const TArray<const FSc
 						PerViewUniformParameters.TangentSelector = FVector4(0.0f, 1.0f, 0.0f, 0.0f);
 					}
 				}
-				PerViewUniformParameters.NormalsType = 0.0f;
-				PerViewUniformParameters.NormalsSphereCenter = FVector4(0.0f, 0.0f, 0.0f, 1.0f);
-				PerViewUniformParameters.NormalsCylinderUnitDirection = FVector4(0.0f, 0.0f, 1.0f, 0.0f);
-				PerViewUniformParameters.PivotOffset = FVector2D(-0.5f, -0.5f);
-				PerViewUniformParameters.MacroUVParameters = FVector4(0.0f, 0.0f, 1.0f, 1.0f);
 
 				// Collector.AllocateOneFrameResource uses default ctor, initialize the vertex factory
 				CollectorResources.VertexFactory.SetFeatureLevel(ViewFamily.GetFeatureLevel());
@@ -388,8 +387,7 @@ void NiagaraEffectRendererRibbon::GetDynamicMeshElements(const TArray<const FSce
 
 	SimpleTimer MeshElementsTimer;
 
-	check(DynamicDataRender)
-	if (DynamicDataRender->VertexData.Num() == 0)
+	if (!DynamicDataRender || DynamicDataRender->VertexData.Num() == 0)
 	{
 		return;
 	}
@@ -542,7 +540,7 @@ FNiagaraDynamicDataBase *NiagaraEffectRendererRibbon::GenerateVertexData(const F
 	TArray<FParticleBeamTrailVertex>& RenderData = DynamicData->VertexData;
 
 	RenderData.Reset(Data.GetNumInstances() * 2);
-	//CachedBounds.Init();
+	CachedBounds.Init();
 
 	// build a sorted list by age, so we always get particles in order 
 	// regardless of them being moved around due to dieing and spawning
@@ -607,6 +605,9 @@ FNiagaraDynamicDataBase *NiagaraEffectRendererRibbon::GenerateVertexData(const F
 		PrevPos = ParticlePos - ParticleRightRot + ParticleDir;
 		PrevPos2 = ParticlePos + ParticleRightRot + ParticleDir;
 		PrevDir = ParticleDir;
+
+		CachedBounds += ParticlePos + ParticleRightRot;
+		CachedBounds += ParticlePos - ParticleRightRot;
 	}
 
 	CPUTimeMS = VertexDataTimer.GetElapsedMilliseconds();
diff --git a/Engine/Source/Runtime/Niagara/Private/NiagaraEmitterProperties.cpp b/Engine/Source/Runtime/Niagara/Private/NiagaraEmitterProperties.cpp
--- a/Engine/Source/Runtime/Niagara/Private/NiagaraEmitterProperties.cpp
+++ b/Engine/Source/Runtime/Niagara/Private/NiagaraEmitterProperties.cpp
@@ -100,45 +100,6 @@ void UNiagaraEmitterProperties::Init()
 	UpdateScriptProps.Init(this);
 }
 
-void UNiagaraEmitterProperties::InitFromOldStruct(FDeprecatedNiagaraEmitterProperties& OldStruct)
-{
-	EmitterName = OldStruct.Name;
-	bIsEnabled = OldStruct.bIsEnabled;
-	SpawnRate = OldStruct.SpawnRate;
-	UpdateScriptProps.Script = OldStruct.UpdateScript;
-	SpawnScriptProps.Script = OldStruct.SpawnScript;
-	Material = OldStruct.Material;
-	RenderModuleType = OldStruct.RenderModuleType;
-	StartTime = OldStruct.StartTime;
-	EndTime = OldStruct.EndTime;
-	RendererProperties = OldStruct.RendererProperties;
-	NumLoops = OldStruct.NumLoops;
-
-#if WITH_EDITORONLY_DATA
-	//Ensure both scripts are post loaded and compile them so that their constant arrays are correct.
-	if (SpawnScriptProps.Script)
-	{
-		SpawnScriptProps.Script->ConditionalPostLoad();
-		if (SpawnScriptProps.Script->Source)
-			SpawnScriptProps.Script->Source->Compile();
-	}
-	if (UpdateScriptProps.Script)
-	{
-		UpdateScriptProps.Script->ConditionalPostLoad();
-		if (UpdateScriptProps.Script->Source)
-			UpdateScriptProps.Script->Source->Compile();
-	}
-#endif
-
-	UpdateScriptProps.ExternalConstants.InitFromOldMap(OldStruct.ExternalConstants);
-	SpawnScriptProps.ExternalConstants.InitFromOldMap(OldStruct.ExternalSpawnConstants);
-	//Also init the spawn constants with the update constants as they used to combined
-	SpawnScriptProps.ExternalConstants.InitFromOldMap(OldStruct.ExternalConstants);
-	UpdateScriptProps.ExternalConstants.InitFromOldMap(OldStruct.ExternalSpawnConstants);
-
-	Init();
-}
-
 #if WITH_EDITOR
 void UNiagaraEmitterProperties::PostEditChangeProperty(struct FPropertyChangedEvent& PropertyChangedEvent)
 {
diff --git a/Engine/Source/Runtime/Niagara/Private/NiagaraSimulation.cpp b/Engine/Source/Runtime/Niagara/Private/NiagaraSimulation.cpp
--- a/Engine/Source/Runtime/Niagara/Private/NiagaraSimulation.cpp
+++ b/Engine/Source/Runtime/Niagara/Private/NiagaraSimulation.cpp
@@ -1,6 +1,7 @@
 // Copyright 1998-2016 Epic Games, Inc. All Rights Reserved.
 
 #include "NiagaraPrivate.h"
+#include "NiagaraConstants.h"
 #include "NiagaraSimulation.h"
 #include "NiagaraEffectRenderer.h"
 #include "VectorVM.h"
@@ -109,9 +110,13 @@ void FNiagaraSimulation::Init()
 	Data.AddVariables(PinnedProps->UpdateScriptProps.Script->Attributes);
 	Data.AddVariables(PinnedProps->SpawnScriptProps.Script->Attributes);
 
-	Constants.Empty();
-	Constants.Merge(PinnedProps->UpdateScriptProps.ExternalConstants);
-	Constants.Merge(PinnedProps->SpawnScriptProps.ExternalConstants);
+//	ExternalConstantsMap.Empty();
+//	ExternalConstantsMap.Merge(PinnedProps->UpdateScriptProps.ExternalConstants);
+//	ExternalConstantsMap.Merge(PinnedProps->SpawnScriptProps.ExternalConstants);
+
+	ExternalConstants.Empty();
+	ExternalConstants.Merge(PinnedProps->UpdateScriptProps.ExternalConstants);
+	ExternalConstants.Merge(PinnedProps->SpawnScriptProps.ExternalConstants);
 
 	//Create data storage for event generators.
 	for (FNiagaraEventGeneratorProperties& Generator : Props->SpawnScriptProps.EventGenerators)
@@ -231,60 +236,76 @@ int FNiagaraSimulation::GetTotalBytesUsed()
 	return BytesUsed;
 }
 
-void FNiagaraSimulation::PreTick()
-{
-	UNiagaraEmitterProperties* PinnedProps = Props.Get();
-	if (!PinnedProps || !bIsEnabled || TickState == NTS_Suspended || TickState == NTS_Dead)
-		return;
 
-	check(Data.GetNumVariables() > 0);
-	check(PinnedProps->SpawnScriptProps.Script);
-	check(PinnedProps->UpdateScriptProps.Script);
+/** Look for dead particles and move from the end of the list to the dead location, compacting in the process
+  */
+void FNiagaraSimulation::KillParticles()
+{
+	SCOPE_CYCLE_COUNTER(STAT_NiagaraKill);
+	int32 OrigNumParticles = Data.GetNumInstances();
+	int32 CurNumParticles = OrigNumParticles;
+	int32 ParticleIndex = OrigNumParticles - 1;
 
-	// Iterate over looking for dead particles and move from the end of the list to the dead location, compacting in the process
+	if (bGenerateDeathEvents)
 	{
-		SCOPE_CYCLE_COUNTER(STAT_NiagaraKill);
-		int32 OrigNumParticles = Data.GetNumInstances();
-		int32 CurNumParticles = OrigNumParticles;
-		int32 ParticleIndex = OrigNumParticles-1;
+		DeathEventGenerator.BeginTrackingDeaths();
+	}
 
-		if (bGenerateDeathEvents)
-		{
-			DeathEventGenerator.BeginTrackingDeaths();
-		}
-		const FVector4* ParticleRelativeTimes = Data.GetVariableData(FNiagaraVariableInfo(FName(TEXT("Age")), ENiagaraDataType::Vector));
-		if (ParticleRelativeTimes)
+	const FVector4* ParticleRelativeTimes = Data.GetVariableData(FNiagaraVariableInfo(BUILTIN_VAR_PARTICLEAGE, ENiagaraDataType::Vector));
+	if (ParticleRelativeTimes)
+	{
+		while (ParticleIndex >= 0)
 		{
-			while (ParticleIndex >= 0)
+			if (ParticleRelativeTimes[ParticleIndex].X > 1.0f)
 			{
-				if (ParticleRelativeTimes[ParticleIndex].X > 1.0f)
+				if (bGenerateDeathEvents)
 				{
-					if (bGenerateDeathEvents)
-					{
-						DeathEventGenerator.OnDeath(ParticleIndex);
-					}
-
-					check(CurNumParticles > ParticleIndex);
-					// Particle is dead, move one from the end here. 
-					MoveParticleToIndex(--CurNumParticles, ParticleIndex);
+					DeathEventGenerator.OnDeath(ParticleIndex);
 				}
-				--ParticleIndex;
+
+				check(CurNumParticles > ParticleIndex);
+
+				// Particle is dead, move one from the end here. 
+				MoveParticleToIndex(--CurNumParticles, ParticleIndex);
 			}
+			--ParticleIndex;
 		}
-		Data.SetNumInstances(CurNumParticles);
+	}
 
-		if (bGenerateDeathEvents)
-		{
-			DeathEventGenerator.EndTrackingDeaths();
-		}
+	Data.SetNumInstances(CurNumParticles);
 
-		// check if the emitter has officially died
-		if (GetTickState() == NTS_Dieing && CurNumParticles == 0)
-		{
-			SetTickState(NTS_Dead);
-		}
+	if (bGenerateDeathEvents)
+	{
+		DeathEventGenerator.EndTrackingDeaths();
 	}
 
+	// check if the emitter has officially died
+	if (GetTickState() == NTS_Dieing && CurNumParticles == 0)
+	{
+		SetTickState(NTS_Dead);
+	}
+}
+
+/** 
+  * PreTick - handles killing dead particles, emitter death, and buffer swaps
+  */
+void FNiagaraSimulation::PreTick()
+{
+	UNiagaraEmitterProperties* PinnedProps = Props.Get();
+
+	if (!PinnedProps || !bIsEnabled
+		|| TickState == NTS_Suspended || TickState == NTS_Dead
+		|| Data.GetNumInstances() == 0)
+	{
+		return;
+	}
+
+	check(Data.GetNumVariables() > 0);
+	check(PinnedProps->SpawnScriptProps.Script);
+	check(PinnedProps->UpdateScriptProps.Script);
+
+	KillParticles();
+
 	//Swap all data set buffers before doing the main tick on any simulation.
 	for (TPair<FNiagaraDataSetID, FNiagaraDataSet*> SetPair : DataSetMap)
 	{
@@ -295,46 +316,36 @@ void FNiagaraSimulation::PreTick()
 void FNiagaraSimulation::Tick(float DeltaSeconds)
 {
 	SCOPE_CYCLE_COUNTER(STAT_NiagaraTick);
+	SimpleTimer TickTime;
 
 	UNiagaraEmitterProperties* PinnedProps = Props.Get();
-	if (!PinnedProps || !bIsEnabled || TickState==NTS_Suspended || TickState==NTS_Dead)
+	if (!PinnedProps || !bIsEnabled || TickState == NTS_Suspended || TickState == NTS_Dead)
+	{
 		return;
+	}
 
-	SimpleTimer TickTime;
+	Age += DeltaSeconds;
 
 	check(Data.GetNumVariables() > 0);
 	check(PinnedProps->SpawnScriptProps.Script);
 	check(PinnedProps->UpdateScriptProps.Script);
 	
-	//Handle Event Actions.
-	auto CallEventActions = [&](FNiagaraEventReceiverProperties& Receiver)
-	{
-		for (auto& Action : Receiver.EmitterActions)
-		{
-			if (Action)
-				Action->PerformAction(*this, Receiver);
-		}
-	};
-	for (auto& Receiver : Props->SpawnScriptProps.EventReceivers)	{ CallEventActions(Receiver); }
-	for (auto& Receiver : Props->UpdateScriptProps.EventReceivers)	{ CallEventActions(Receiver); }
-
-	int32 OrigNumParticles = Data.GetNumInstances();
-	int32 NumToSpawn = 0;
+	TickEvents(DeltaSeconds);
 
 	// Figure out how many we will spawn.
-	NumToSpawn = CalcNumToSpawn(DeltaSeconds);
+	int32 OrigNumParticles = Data.GetNumInstances();
+	int32 NumToSpawn = CalcNumToSpawn(DeltaSeconds);
 	int32 MaxNewParticles = OrigNumParticles + NumToSpawn;
 	Data.Allocate(MaxNewParticles);
 
-	Age += DeltaSeconds;
-	Constants.SetOrAdd(TEXT("Emitter Age"), FVector4(Age, Age, Age, Age));
-	Constants.SetOrAdd(TEXT("Delta Time"), FVector4(DeltaSeconds, DeltaSeconds, DeltaSeconds, DeltaSeconds));
+	ExternalConstants.SetOrAdd(BUILTIN_CONST_EMITTERAGE, FVector4(Age, Age, Age, Age));
+	ExternalConstants.SetOrAdd(BUILTIN_CONST_DELTATIME, FVector4(DeltaSeconds, DeltaSeconds, DeltaSeconds, DeltaSeconds));
 
 	// Simulate particles forward by DeltaSeconds.
 	if (TickState==NTS_Running || TickState==NTS_Dieing)
 	{
 		SCOPE_CYCLE_COUNTER(STAT_NiagaraSimulate);
-		RunVMScript(PinnedProps->UpdateScriptProps, EUnusedAttributeBehaviour::Copy);
+		RunVMScript(PinnedProps->UpdateScriptProps, EUnusedAttributeBehaviour::PassThrough);
 	}
 
 	//Init new particles with the spawn script.
@@ -344,20 +355,39 @@ void FNiagaraSimulation::Tick(float DeltaSeconds)
 		Data.SetNumInstances(MaxNewParticles);
 		//For now, zero any unused attributes here. But as this is really uninitialized data we should maybe make this a more serious error.
 		RunVMScript(PinnedProps->SpawnScriptProps, EUnusedAttributeBehaviour::Zero, OrigNumParticles, NumToSpawn);
-	}
 
-	if (bGenerateSpawnEvents)
-	{
-		SpawnEventGenerator.OnSpawned(OrigNumParticles, NumToSpawn);
+		if (bGenerateSpawnEvents)
+		{
+			SpawnEventGenerator.OnSpawned(OrigNumParticles, NumToSpawn);
+		}
 	}
 
 
 	CPUTimeMS = TickTime.GetElapsedMilliseconds();
 
 	INC_DWORD_STAT_BY(STAT_NiagaraNumParticles, Data.GetNumInstances());
+}
+
+
+void FNiagaraSimulation::TickEvents(float DeltaSeconds)
+{
+	//Handle Event Actions.
+	auto CallEventActions = [&](FNiagaraEventReceiverProperties& Receiver)
+	{
+		for (auto& Action : Receiver.EmitterActions)
+		{
+			if (Action)
+			{
+				Action->PerformAction(*this, Receiver);
+			}
+		}
+	};
 
+	for (auto& Receiver : Props->SpawnScriptProps.EventReceivers) { CallEventActions(Receiver); }
+	for (auto& Receiver : Props->UpdateScriptProps.EventReceivers) { CallEventActions(Receiver); }
 }
 
+
 void FNiagaraSimulation::RunVMScript(FNiagaraEmitterScriptProperties& ScriptProps, EUnusedAttributeBehaviour UnusedAttribBehaviour)
 {
 	RunVMScript(ScriptProps, UnusedAttribBehaviour, 0, Data.GetNumInstances());
@@ -371,7 +401,9 @@ void FNiagaraSimulation::RunVMScript(FNiagaraEmitterScriptProperties& ScriptProp
 void FNiagaraSimulation::RunVMScript(FNiagaraEmitterScriptProperties& ScriptProps, EUnusedAttributeBehaviour UnusedAttribBehaviour, uint32 StartParticle, uint32 NumParticles)
 {
 	if (NumParticles == 0)
+	{
 		return;
+	}
 	
 	UNiagaraScript* Script = ScriptProps.Script;
 
@@ -384,18 +416,11 @@ void FNiagaraSimulation::RunVMScript(FNiagaraEmitterScriptProperties& ScriptProp
 
 	check(StartParticle + NumParticles <= CurrNumParticles);
 
-	const FVector4* InputBuffer = Data.GetPreviousBuffer();
-	const FVector4* OutputBuffer = Data.GetCurrentBuffer();
-
-	uint32 InputAttributeStride = Data.GetPrevDataAllocation();
-	uint32 OutputAttributeStride = Data.GetDataAllocation();
-
 	VectorRegister* InputRegisters[VectorVM::MaxInputRegisters] = { 0 };
 	VectorRegister* OutputRegisters[VectorVM::MaxOutputRegisters] = { 0 };
 	
 	//The current script run will be using NumScriptAttrs. We must pick out the Attributes form the simulation data in the order that they appear in the script.
 	//The remaining attributes being handled according to UnusedAttribBehaviour
-	const int32 NumSimulationAttrs = Data.GetNumVariables();
 	const int32 NumScriptAttrs = Script->Attributes.Num();
 
 	check(NumScriptAttrs < VectorVM::MaxInputRegisters);
@@ -403,55 +428,65 @@ void FNiagaraSimulation::RunVMScript(FNiagaraEmitterScriptProperties& ScriptProp
 
 	int32 NumInputRegisters = 0;
 	int32 NumOutputRegitsers = 0;
-	// Setup input and output registers.
-	for (TPair<FNiagaraVariableInfo, uint32> AttrIndexPair : Data.GetVariables())
-	{
-		const FNiagaraVariableInfo& AttrInfo = AttrIndexPair.Key;
 
-		FVector4* InBuff;
-		FVector4* OutBuff;
-		Data.GetVariableData(AttrInfo, InBuff, OutBuff, StartParticle);
+	{
+		SCOPE_CYCLE_COUNTER(STAT_NiagaraSimRegisterSetup);
 
-		int32 AttrIndex = Script->Attributes.Find(AttrInfo);
-		if (AttrIndex != INDEX_NONE)
-		{
-			//The input buffer can be NULL but only if this script never accesses it.
-			check(InBuff || !Script->Usage.bReadsAttriubteData);
-
-			//This attribute is used in this script so set it to the correct Input and Output buffer locations.
-			InputRegisters[AttrIndex] = (VectorRegister*)InBuff;
-			OutputRegisters[AttrIndex] = (VectorRegister*)OutBuff;
-			
-			NumInputRegisters += InBuff ? 1 : 0;
-			NumOutputRegitsers += OutBuff ? 1 : 0;
-		}	
-		else
+		// Setup input and output registers.
+		for (TPair<FNiagaraVariableInfo, uint32> AttrIndexPair : Data.GetVariables())
 		{
-			//This attribute is not used in this script so handle it accordingly.
-			check(OutBuff);
+			const FNiagaraVariableInfo& AttrInfo = AttrIndexPair.Key;
 
-			if (UnusedAttribBehaviour == EUnusedAttributeBehaviour::Copy)
-			{
-				check(InBuff);
-				FMemory::Memcpy(OutBuff, InBuff, NumParticles * sizeof(FVector4));
-			}
-			else if (UnusedAttribBehaviour == EUnusedAttributeBehaviour::Zero)
+			FVector4* InBuff;
+			FVector4* OutBuff;
+			Data.GetVariableData(AttrInfo, InBuff, OutBuff, StartParticle);
+
+			int32 AttrIndex = Script->Attributes.Find(AttrInfo);
+			if (AttrIndex != INDEX_NONE)
 			{
-				FMemory::Memzero(OutBuff, NumParticles * sizeof(FVector4));
+				//The input buffer can be NULL but only if this script never accesses it.
+				check(InBuff || !Script->Usage.bReadsAttriubteData);
+
+				//This attribute is used in this script so set it to the correct Input and Output buffer locations.
+				InputRegisters[AttrIndex] = (VectorRegister*)InBuff;
+				OutputRegisters[AttrIndex] = (VectorRegister*)OutBuff;
+
+				NumInputRegisters += InBuff ? 1 : 0;
+				NumOutputRegitsers += OutBuff ? 1 : 0;
 			}
-			else if (UnusedAttribBehaviour == EUnusedAttributeBehaviour::MarkInvalid)
+			else
 			{
-				FMemory::Memset(OutBuff, NIAGARA_INVALID_MEMORY, NumParticles * sizeof(FVector4));
+				//This attribute is not used in this script so handle it accordingly.
+				check(OutBuff);
+
+				if (UnusedAttribBehaviour == EUnusedAttributeBehaviour::PassThrough)
+				{
+					Data.GetVariableBuffer(AttrInfo).EnablePassThrough();
+				}
+				else if (UnusedAttribBehaviour == EUnusedAttributeBehaviour::Copy)
+				{
+					check(InBuff);
+					FMemory::Memcpy(OutBuff, InBuff, NumParticles * sizeof(FVector4));
+				}
+				else if (UnusedAttribBehaviour == EUnusedAttributeBehaviour::Zero)
+				{
+					FMemory::Memzero(OutBuff, NumParticles * sizeof(FVector4));
+				}
+				else if (UnusedAttribBehaviour == EUnusedAttributeBehaviour::MarkInvalid)
+				{
+					FMemory::Memset(OutBuff, NIAGARA_INVALID_MEMORY, NumParticles * sizeof(FVector4));
+				}
 			}
 		}
 	}
 
-	//Fill constant table with required emitter constants and internal script constants.
+	// Fill constant table with required emitter constants and internal script constants.
 	TArray<FVector4> ConstantTable;
-	Script->ConstantData.FillConstantTable(Constants, ConstantTable);
+	//Script->ConstantData.FillConstantTable(ExternalConstantsMap, ConstantTable);
+	Script->ConstantData.FillConstantTable(ExternalConstants, ConstantTable);
 
 
-	//Fill in the shared data table
+	// Fill in the shared data table
 	TArray<FVectorVMSharedDataView, TInlineAllocator<64>> SharedDataTable;
 	
 	for (FNiagaraEventReceiverProperties Receiver : ScriptProps.EventReceivers)
diff --git a/Engine/Source/Runtime/Niagara/Public/NiagaraCommon.h b/Engine/Source/Runtime/Niagara/Public/NiagaraCommon.h
--- a/Engine/Source/Runtime/Niagara/Public/NiagaraCommon.h
+++ b/Engine/Source/Runtime/Niagara/Public/NiagaraCommon.h
@@ -15,6 +15,19 @@ enum class ENiagaraDataType : uint8
 };
 
 
+enum class ENiagaraType : uint8
+{
+	Byte,
+	Word,
+	DWord,
+	Float 
+};
+
+enum ENiagaraSimTarget
+{
+	CPUSim,
+	GPUComputeSim
+};
 
 
 enum ENiagaraTickState
@@ -26,6 +39,11 @@ enum ENiagaraTickState
 };
 
 
+uint32 GetNiagaraDataElementCount(ENiagaraDataType Type);
+uint32 GetNiagaraBytesPerElement(ENiagaraDataType DataType, ENiagaraType Type);
+
+
+
 USTRUCT()
 struct FNiagaraVariableInfo
 {
diff --git a/Engine/Source/Runtime/Niagara/Public/NiagaraFunctionLibrary.h b/Engine/Source/Runtime/Niagara/Public/NiagaraFunctionLibrary.h
--- a/Engine/Source/Runtime/Niagara/Public/NiagaraFunctionLibrary.h
+++ b/Engine/Source/Runtime/Niagara/Public/NiagaraFunctionLibrary.h
@@ -7,7 +7,7 @@
 
 
 /**
-* A C++ and Blueprint accessible library of utility functions for accessing Leap Motion device.
+* A C++ and Blueprint accessible library of utility functions for accessing Niagara simulations
 * All positions & orientations are returned in Unreal reference frame & units, assuming the Leap device is located at the origin.
 */
 UCLASS()
diff --git a/Engine/Source/Runtime/OpenGLDrv/Private/OpenGLCommands.cpp b/Engine/Source/Runtime/OpenGLDrv/Private/OpenGLCommands.cpp
--- a/Engine/Source/Runtime/OpenGLDrv/Private/OpenGLCommands.cpp
+++ b/Engine/Source/Runtime/OpenGLDrv/Private/OpenGLCommands.cpp
@@ -1837,10 +1837,15 @@ void FOpenGLDynamicRHI::RHISetRenderTargets(
 
 void FOpenGLDynamicRHI::RHIDiscardRenderTargets(bool Depth, bool Stencil, uint32 ColorBitMask)
 {
-	if(FOpenGL::SupportsDiscardFrameBuffer())
+	if (FOpenGL::SupportsDiscardFrameBuffer())
 	{
+		{
+			QUICK_SCOPE_CYCLE_COUNTER(STAT_RHIMETHOD_DiscardRenderTargets_Flush);
+			FRHICommandListExecutor::GetImmediateCommandList().ImmediateFlush(EImmediateFlushType::FlushRHIThread);
+		}
+
 		// 8 Color + Depth + Stencil = 10
-		GLenum Attachments[10];
+		GLenum Attachments[MaxSimultaneousRenderTargets + 2];
 		uint32 I=0;
 		if(Depth) 
 		{
@@ -1853,7 +1858,7 @@ void FOpenGLDynamicRHI::RHIDiscardRenderTargets(bool Depth, bool Stencil, uint32
 			I++;
 		}
 
-		ColorBitMask &= (1 << 8) - 1;
+		ColorBitMask &= (1 << MaxSimultaneousRenderTargets) - 1;
 		uint32 J = 0;
 		while (ColorBitMask)
 		{
diff --git a/Engine/Source/Runtime/OpenGLDrv/Private/OpenGLDevice.cpp b/Engine/Source/Runtime/OpenGLDrv/Private/OpenGLDevice.cpp
--- a/Engine/Source/Runtime/OpenGLDrv/Private/OpenGLDevice.cpp
+++ b/Engine/Source/Runtime/OpenGLDrv/Private/OpenGLDevice.cpp
@@ -851,6 +851,7 @@ static void InitRHICapabilitiesForGL()
 		SetupTextureFormat( PF_R32_SINT,		FOpenGLTextureFormat( GL_R32I,					GL_NONE,				GL_RED_INTEGER,	GL_INT,								false,	false));
 		SetupTextureFormat( PF_R16_UINT,		FOpenGLTextureFormat( GL_R16UI,					GL_NONE,				GL_RED_INTEGER,	GL_UNSIGNED_SHORT,					false,	false));
 		SetupTextureFormat( PF_R16_SINT,		FOpenGLTextureFormat( GL_R16I,					GL_NONE,				GL_RED_INTEGER,	GL_SHORT,							false,	false));
+		SetupTextureFormat( PF_R8_UINT,			FOpenGLTextureFormat( GL_R8UI,					GL_NONE,				GL_RED_INTEGER, GL_UNSIGNED_BYTE,					false,  false));
 		SetupTextureFormat( PF_FloatRGBA,		FOpenGLTextureFormat( GL_RGBA16F,				GL_RGBA16F,				GL_RGBA,		GL_HALF_FLOAT,						false,	false));
 		if (FOpenGL::GetShaderPlatform() == EShaderPlatform::SP_OPENGL_ES31_EXT)
 		{
diff --git a/Engine/Source/Runtime/OpenGLDrv/Private/OpenGLShaders.cpp b/Engine/Source/Runtime/OpenGLDrv/Private/OpenGLShaders.cpp
--- a/Engine/Source/Runtime/OpenGLDrv/Private/OpenGLShaders.cpp
+++ b/Engine/Source/Runtime/OpenGLDrv/Private/OpenGLShaders.cpp
@@ -9,6 +9,7 @@
 #include "Shader.h"
 #include "GlobalShader.h"
 #include "ShaderCache.h"
+#include "CrossCompilerCommon.h"
 
 #define CHECK_FOR_GL_SHADERS_TO_REPLACE 0
 
@@ -1204,9 +1205,8 @@ void FOpenGLLinkedProgram::VerifyUniformBlockBindings( int Stage, uint32 FirstUn
 {
 	if ( FOpenGL::SupportsSeparateShaderObjects() && FOpenGL::SupportsUniformBuffers() )
 	{
-		static const ANSICHAR StagePrefix[CrossCompiler::NUM_SHADER_STAGES] = { 'v', 'p', 'g', 'h', 'd', 'c'};
 		FOpenGLUniformName Name;
-		Name.Buffer[0] = StagePrefix[Stage];
+		Name.Buffer[0] = CrossCompiler::ShaderStageIndexToTypeName(Stage);
 		Name.Buffer[1] = 'b';
 		
 		GLuint StageProgram = Config.Shaders[Stage].Resource;
@@ -1225,7 +1225,6 @@ void FOpenGLLinkedProgram::VerifyUniformBlockBindings( int Stage, uint32 FirstUn
 
 void FOpenGLLinkedProgram::ConfigureShaderStage( int Stage, uint32 FirstUniformBuffer )
 {
-	static const ANSICHAR StagePrefix[CrossCompiler::NUM_SHADER_STAGES] = { 'v', 'p', 'g', 'h', 'd', 'c'};
 	static const GLint FirstTextureUnit[CrossCompiler::NUM_SHADER_STAGES] =
 	{
 		FOpenGL::GetFirstVertexTextureUnit(),
@@ -1252,7 +1251,7 @@ void FOpenGLLinkedProgram::ConfigureShaderStage( int Stage, uint32 FirstUniformB
 	VERIFY_GL_SCOPE();
 
 	FOpenGLUniformName Name;
-	Name.Buffer[0] = StagePrefix[Stage];
+	Name.Buffer[0] = CrossCompiler::ShaderStageIndexToTypeName(Stage);
 
 	GLuint StageProgram = FOpenGL::SupportsSeparateShaderObjects() ? Config.Shaders[Stage].Resource : Program;
 	
